Added DataLayer browser & editing.

diff --git a/Engine/Source/Editor/DataLayerEditor/DataLayerEditor.Build.cs b/Engine/Source/Editor/DataLayerEditor/DataLayerEditor.Build.cs
new file mode 100644
--- /dev/null
+++ b/Engine/Source/Editor/DataLayerEditor/DataLayerEditor.Build.cs
@@ -0,0 +1,33 @@
+// Copyright Epic Games, Inc. All Rights Reserved.
+
+using UnrealBuildTool;
+
+public class DataLayerEditor : ModuleRules
+{
+	public DataLayerEditor(ReadOnlyTargetRules Target) : base(Target)
+	{
+		PublicIncludePaths.Add("Editor/DataLayerEditor/Public");
+
+		PrivateIncludePaths.Add("Editor/DataLayerEditor/Private"); // For PCH includes (because they don't work with relative paths, yet)
+
+		PrivateDependencyModuleNames.AddRange(
+			new string[] {
+				"Core",
+				"CoreUObject",
+				"EditorFramework",
+				"EditorWidgets",
+				"EditorStyle",
+				"EditorSubsystem",
+				"PropertyEditor",
+				"Engine",
+                "InputCore",
+				"Slate",
+				"SlateCore",
+				"UnrealEd",
+				"SceneOutliner",
+				"ToolMenus",
+				"AssetTools",
+			}
+		);
+	}
+}
diff --git a/Engine/Source/Editor/DataLayerEditor/Private/DataLayer/DataLayerActorTreeItem.h b/Engine/Source/Editor/DataLayerEditor/Private/DataLayer/DataLayerActorTreeItem.h
new file mode 100644
--- /dev/null
+++ b/Engine/Source/Editor/DataLayerEditor/Private/DataLayer/DataLayerActorTreeItem.h
@@ -0,0 +1,72 @@
+// Copyright Epic Games, Inc. All Rights Reserved.
+
+#pragma once
+
+#include "CoreMinimal.h"
+#include "SceneOutlinerFwd.h"
+#include "ActorTreeItem.h"
+#include "WorldPartition/DataLayer/DataLayer.h"
+
+//////////////////////////////////////////////////////////////////////////
+// FDataLayerActorTreeItemData
+
+struct FDataLayerActorTreeItemData
+{
+	FDataLayerActorTreeItemData(AActor* InActor, UDataLayer* InDataLayer)
+		: Actor(InActor)
+		, DataLayer(InDataLayer)
+	{}
+	TWeakObjectPtr<AActor> Actor;
+	TWeakObjectPtr<UDataLayer> DataLayer;
+};
+
+//////////////////////////////////////////////////////////////////////////
+// FDataLayerActorTreeItem
+
+struct FDataLayerActorTreeItem : public FActorTreeItem
+{
+public:
+	DECLARE_DELEGATE_RetVal_TwoParams(bool, FFilterPredicate, const AActor*, const UDataLayer* DataLayer);
+	DECLARE_DELEGATE_RetVal_TwoParams(bool, FInteractivePredicate, const AActor*, const UDataLayer* DataLayer);
+
+	FDataLayerActorTreeItem(const FDataLayerActorTreeItemData& InData)
+		: FActorTreeItem(InData.Actor.Get())
+		, DataLayer(InData.DataLayer)
+		, IDDataLayerActor(FDataLayerActorTreeItem::ComputeTreeItemID(Actor.Get(), DataLayer.Get()))
+	{
+	}
+
+	const UDataLayer* GetDataLayer() const { return DataLayer.Get(); }
+	UDataLayer* GetDataLayer() { return DataLayer.Get(); }
+
+	const AActor* GetActor() const { return Actor.Get(); }
+	AActor* GetActor() { return Actor.Get(); }
+
+	static uint32 ComputeTreeItemID(const AActor* InActor, const UDataLayer* InDataLayer)
+	{
+		return HashCombine(GetTypeHash(FObjectKey(InActor)), GetTypeHash(FObjectKey(InDataLayer)));
+	}
+
+	bool Filter(FFilterPredicate Pred) const
+	{
+		return Pred.Execute(Actor.Get(), DataLayer.Get());
+	}
+
+	bool GetInteractiveState(FInteractivePredicate Pred) const
+	{
+		return Pred.Execute(Actor.Get(), DataLayer.Get());
+	}
+
+	/* Begin ISceneOutlinerTreeItem Implementation */
+	virtual bool IsValid() const override { return Actor.IsValid() && DataLayer.IsValid(); }
+	virtual FSceneOutlinerTreeItemID GetID() const override { return IDDataLayerActor; }
+	virtual bool ShouldShowVisibilityState() const { return false; }
+	virtual bool HasVisibilityInfo() const override { return false; }
+	virtual void OnVisibilityChanged(const bool bNewVisibility) override {}
+	virtual bool GetVisibility() const override { return false; }
+	/* End ISceneOutlinerTreeItem Implementation */
+
+private:
+	TWeakObjectPtr<UDataLayer> DataLayer;
+	const uint32 IDDataLayerActor;
+};
\ No newline at end of file
diff --git a/Engine/Source/Editor/DataLayerEditor/Private/DataLayer/DataLayerDragDropOp.cpp b/Engine/Source/Editor/DataLayerEditor/Private/DataLayer/DataLayerDragDropOp.cpp
new file mode 100644
--- /dev/null
+++ b/Engine/Source/Editor/DataLayerEditor/Private/DataLayer/DataLayerDragDropOp.cpp
@@ -0,0 +1,21 @@
+// Copyright Epic Games, Inc. All Rights Reserved.
+
+#include "DataLayer/DataLayerDragDropOp.h"
+#include "Textures/SlateIcon.h"
+
+void FDataLayerDragDropOp::Construct()
+{
+	const FSlateBrush* Icon = FEditorStyle::GetBrush(TEXT("DataLayer.Icon16x"));
+	if (DataLayerLabels.Num() == 1)
+	{
+		SetToolTip(FText::FromName(DataLayerLabels[0]), Icon);
+	}
+	else
+	{
+		FText Text = FText::Format(NSLOCTEXT("FDataLayerDragDropOp", "MultipleFormat", "{0} DataLayerLabels"), DataLayerLabels.Num());
+		SetToolTip(Text, Icon);
+	}
+
+	SetupDefaults();
+	FDecoratedDragDropOp::Construct();
+}
diff --git a/Engine/Source/Editor/DataLayerEditor/Private/DataLayer/DataLayerDragDropOp.h b/Engine/Source/Editor/DataLayerEditor/Private/DataLayer/DataLayerDragDropOp.h
new file mode 100644
--- /dev/null
+++ b/Engine/Source/Editor/DataLayerEditor/Private/DataLayer/DataLayerDragDropOp.h
@@ -0,0 +1,19 @@
+// Copyright Epic Games, Inc. All Rights Reserved.
+
+#pragma once
+
+#include "DragAndDrop/DecoratedDragDropOp.h"
+
+class FDragDropEvent;
+
+/** Drag/drop operation for dragging layers in the editor */
+class FDataLayerDragDropOp : public FDecoratedDragDropOp
+{
+public:
+	DRAG_DROP_OPERATOR_TYPE(FDataLayerDragDropOp, FDragDropOperation)
+
+	/** The labels of the layers being dragged */
+	TArray<FName> DataLayerLabels;
+
+	virtual void Construct() override;
+};
diff --git a/Engine/Source/Editor/DataLayerEditor/Private/DataLayer/DataLayerEditorSubsystem.cpp b/Engine/Source/Editor/DataLayerEditor/Private/DataLayer/DataLayerEditorSubsystem.cpp
new file mode 100644
--- /dev/null
+++ b/Engine/Source/Editor/DataLayerEditor/Private/DataLayer/DataLayerEditorSubsystem.cpp
@@ -0,0 +1,1193 @@
+// Copyright Epic Games, Inc. All Rights Reserved.
+
+#include "DataLayer/DataLayerEditorSubsystem.h"
+#include "Components/PrimitiveComponent.h"
+#include "WorldPartition/DataLayer/DataLayer.h"
+#include "WorldPartition/DataLayer/WorldDataLayers.h"
+#include "WorldPartition/WorldPartitionSubsystem.h"
+#include "Modules/ModuleManager.h"
+#include "DataLayerEditorModule.h"
+#include "ActorEditorUtils.h"
+#include "LevelEditorViewport.h"
+#include "Algo/Transform.h"
+#include "Misc/IFilter.h"
+#include "Engine/Selection.h"
+#include "Engine/World.h"
+#include "Engine/Brush.h"
+#include "EngineUtils.h"
+#include "Editor.h"
+
+//////////////////////////////////////////////////////////////////////////
+// FDataLayersBroadcast
+
+class FDataLayersBroadcast
+{
+public:
+	FDataLayersBroadcast(UDataLayerEditorSubsystem* InDataLayerEditorSubsystem);
+	~FDataLayersBroadcast();
+	void Deinitialize();
+
+private:
+	void Initialize();
+	void OnEditorMapChange(uint32 MapChangeFlags = 0) { DataLayerEditorSubsystem->EditorMapChange(); }
+	void OnPostUndoRedo() { DataLayerEditorSubsystem->PostUndoRedo(); }
+	void OnObjectPostEditChange(UObject* Object, FPropertyChangedEvent& PropertyChangedEvent);
+	void OnLevelViewportClientListChanged();
+	void OnLevelActorsAdded(AActor* InActor) { DataLayerEditorSubsystem->InitializeNewActorDataLayers(InActor); }
+
+	TSet<FLevelEditorViewportClient*> RegisteredViewportClients;
+	UDataLayerEditorSubsystem* DataLayerEditorSubsystem;
+	bool bIsInitialized;
+};
+
+FDataLayersBroadcast::FDataLayersBroadcast(UDataLayerEditorSubsystem* InDataLayerEditorSubsystem)
+	: DataLayerEditorSubsystem(InDataLayerEditorSubsystem)
+	, bIsInitialized(false)
+{
+	Initialize();
+}
+
+FDataLayersBroadcast::~FDataLayersBroadcast()
+{
+	Deinitialize();
+}
+
+void FDataLayersBroadcast::Deinitialize()
+{
+	if (bIsInitialized)
+	{
+		bIsInitialized = false;
+		FEditorDelegates::MapChange.RemoveAll(this);
+		FEditorDelegates::PostUndoRedo.RemoveAll(this);
+		FCoreUObjectDelegates::OnObjectPropertyChanged.RemoveAll(this);
+		if (GEditor)
+		{
+			GEditor->OnLevelViewportClientListChanged().RemoveAll(this);
+			RegisteredViewportClients.Reset();
+		}
+		GEngine->OnLevelActorAdded().RemoveAll(this);
+		GEngine->OnLevelActorDeleted().RemoveAll(this);
+	}
+}
+
+void FDataLayersBroadcast::Initialize()
+{
+	if (!bIsInitialized)
+	{
+		bIsInitialized = true;
+		FEditorDelegates::MapChange.AddRaw(this, &FDataLayersBroadcast::OnEditorMapChange);
+		FEditorDelegates::PostUndoRedo.AddRaw(this, &FDataLayersBroadcast::OnPostUndoRedo);
+		FCoreUObjectDelegates::OnObjectPropertyChanged.AddRaw(this, &FDataLayersBroadcast::OnObjectPostEditChange);
+		if (GEditor)
+		{
+			RegisteredViewportClients.Reset();
+			RegisteredViewportClients.Append(GEditor->GetLevelViewportClients());
+			for (FLevelEditorViewportClient* ViewportClient : RegisteredViewportClients)
+			{
+				DataLayerEditorSubsystem->UpdatePerViewVisibility(ViewportClient);
+			}
+			GEditor->OnLevelViewportClientListChanged().AddRaw(this, &FDataLayersBroadcast::OnLevelViewportClientListChanged);
+			GEngine->OnLevelActorAdded().AddRaw(this, &FDataLayersBroadcast::OnLevelActorsAdded);
+		}
+	}
+}
+
+void FDataLayersBroadcast::OnObjectPostEditChange(UObject* Object, FPropertyChangedEvent& PropertyChangedEvent)
+{
+	if (Object)
+	{
+		bool bRefresh = false;
+		if (UDataLayer* DataLayer = Cast<UDataLayer>(Object))
+		{
+			bRefresh = true;
+		}
+		else if (AActor* Actor = Cast<AActor>(Object))
+		{
+			bRefresh = Actor->IsPropertyChangedAffectingDataLayers(PropertyChangedEvent);
+		}
+		if (bRefresh)
+		{
+			// Force and update
+			DataLayerEditorSubsystem->EditorRefreshDataLayerBrowser();
+		}
+	}
+}
+
+void FDataLayersBroadcast::OnLevelViewportClientListChanged()
+{
+	if (GEditor)
+	{
+		TSet<FLevelEditorViewportClient*> NewViewportClients(GEditor->GetLevelViewportClients());
+		TSet<FLevelEditorViewportClient*> AddedViewportClients = NewViewportClients.Difference(RegisteredViewportClients);
+		TSet<FLevelEditorViewportClient*> RemovedViewportClients = RegisteredViewportClients.Difference(NewViewportClients);
+		for (FLevelEditorViewportClient* ViewportClient : AddedViewportClients)
+		{
+			DataLayerEditorSubsystem->UpdatePerViewVisibility(ViewportClient);
+		}
+		for (FLevelEditorViewportClient* ViewportClient : RemovedViewportClients)
+		{
+			DataLayerEditorSubsystem->RemoveViewFromActorViewVisibility(ViewportClient);
+		}
+		RegisteredViewportClients = NewViewportClients;
+	}
+}
+
+//////////////////////////////////////////////////////////////////////////
+// UDataLayerEditorSubsystem
+//
+// Note: 
+//		- DataLayer visibility currently re-uses Actor's bHiddenEdLayer and HiddenEditorViews. It's viable since Layer & DataLayer ares mutually exclusive systems.
+//		- UDataLayerEditorSubsystem is intended to replace ULayersSubsystem for worlds using the World Partition system.
+//		  Extra work is necessary to replace all references to GetEditorSubsystem<ULayersSubsystem> in the Editor.
+//		  Either a proxy that redirects calls to the proper EditorSubsystem will be used or user code will change to trigger delegate broadcast instead of directly accessing the subsystem (see calls to InitializeNewActorDataLayers everywhere as an example).
+//
+
+UDataLayerEditorSubsystem* UDataLayerEditorSubsystem::Get()
+{
+	if (!FModuleManager::Get().IsModuleLoaded("DataLayerEditor"))
+	{
+		FModuleManager::LoadModuleChecked<FDataLayerEditorModule>("DataLayerEditor");
+	}
+	return GEditor ? GEditor->GetEditorSubsystem<UDataLayerEditorSubsystem>() : nullptr;
+}
+
+void UDataLayerEditorSubsystem::Initialize(FSubsystemCollectionBase& Collection)
+{
+	Super::Initialize(Collection);
+
+	// Set up the broadcast functions for DataLayerEditorSubsystem
+	DataLayersBroadcast = MakeShareable(new FDataLayersBroadcast(this));
+
+	if (UWorld* World = GetWorld())
+	{
+		World->PersistentLevel->OnLoadedActorAddedToLevelEvent.AddLambda([this](AActor& InActor) { InitializeNewActorDataLayers(&InActor); });
+	}
+}
+
+void UDataLayerEditorSubsystem::Deinitialize()
+{
+	Super::Deinitialize();
+
+	DataLayersBroadcast->Deinitialize();
+}
+
+void UDataLayerEditorSubsystem::EditorMapChange()
+{
+	if (UWorld * World = GetWorld())
+	{
+		World->PersistentLevel->OnLoadedActorAddedToLevelEvent.AddLambda([this](AActor& InActor) { InitializeNewActorDataLayers(&InActor); });
+	}
+	DataLayerChanged.Broadcast(EDataLayerAction::Reset, NULL, NAME_None);
+}
+
+void UDataLayerEditorSubsystem::EditorRefreshDataLayerBrowser()
+{
+	DataLayerChanged.Broadcast(EDataLayerAction::Reset, NULL, NAME_None);
+	UpdateAllActorsVisibility(false, false);
+}
+
+void UDataLayerEditorSubsystem::PostUndoRedo()
+{
+	DataLayerChanged.Broadcast(EDataLayerAction::Reset, NULL, NAME_None);
+	UpdateAllActorsVisibility(true, true);
+}
+
+///////////////////////////////////////////////////////////////////////////////////////////////////////////////////////
+//
+// Operations on an individual actor.
+//
+///////////////////////////////////////////////////////////////////////////////////////////////////////////////////////
+
+bool UDataLayerEditorSubsystem::IsActorValidForDataLayer(AActor* Actor)
+{
+	if (!Actor || Actor->GetClass() == nullptr)
+	{
+		return false;
+	}
+
+	UWorld* World = Actor->GetWorld();
+	if (!World)
+	{
+		return false;
+	}
+
+	const bool bIsPartitionedActor = UWorld::HasSubsystem<UWorldPartitionSubsystem>(World);
+	const bool bIsInEditorWorld = World->WorldType == EWorldType::Editor;
+	const bool bIsBuilderBrush = FActorEditorUtils::IsABuilderBrush(Actor);
+	const bool bIsHidden = Actor->GetClass()->GetDefaultObject<AActor>()->bHiddenEd;
+	const bool bIsValid = !bIsHidden && !bIsBuilderBrush && bIsInEditorWorld && bIsPartitionedActor;
+
+	return bIsValid;
+}
+
+void UDataLayerEditorSubsystem::InitializeNewActorDataLayers(AActor* Actor)
+{
+	if (!IsActorValidForDataLayer(Actor))
+	{
+		return;
+	}
+
+	Actor->FixupDataLayers();
+
+	// update per-view visibility info
+	UpdateActorAllViewsVisibility(Actor);
+
+	// update general actor visibility
+	bool bActorModified = false;
+	bool bActorSelectionChanged = false;
+	UpdateActorVisibility(Actor, bActorSelectionChanged, bActorModified, /*bActorNotifySelectionChange*/true, /*bActorRedrawViewports*/false);
+}
+
+UWorld* UDataLayerEditorSubsystem::GetWorld() const
+{
+	return GWorld;
+}
+
+bool UDataLayerEditorSubsystem::AddActorToDataLayer(AActor* Actor, const UDataLayer* DataLayer)
+{
+	TArray<AActor*> Actors;
+	Actors.Add(Actor);
+
+	return AddActorsToDataLayers(Actors, {DataLayer});
+}
+
+bool UDataLayerEditorSubsystem::AddActorToDataLayers(AActor* Actor, const TArray<const UDataLayer*>& DataLayers)
+{
+	TArray<AActor*> Actors;
+	Actors.Add(Actor);
+
+	return AddActorsToDataLayers(Actors, DataLayers);
+}
+
+bool UDataLayerEditorSubsystem::AddActorsToDataLayer(const TArray<AActor*>& Actors, const UDataLayer* DataLayer)
+{
+	return AddActorsToDataLayers(Actors, {DataLayer});
+}
+
+bool UDataLayerEditorSubsystem::AddActorsToDataLayer(const TArray<TWeakObjectPtr<AActor>>& Actors, const UDataLayer* DataLayer)
+{
+	return AddActorsToDataLayers(Actors, {DataLayer});
+}
+
+bool UDataLayerEditorSubsystem::AddActorsToDataLayers(const TArray<AActor*>& Actors, const TArray<const UDataLayer*>& DataLayers)
+{
+	bool bChangesOccurred = false;
+
+	if (DataLayers.Num() > 0)
+	{
+		GEditor->GetSelectedActors()->BeginBatchSelectOperation();
+
+		for (AActor* Actor : Actors)
+		{
+			if (!IsActorValidForDataLayer(Actor))
+			{
+				continue;
+			}
+
+			bool bActorWasModified = false;
+			for (const UDataLayer* DataLayer : DataLayers)
+			{
+				if (Actor->AddDataLayer(DataLayer))
+				{
+					bActorWasModified = true;
+					ActorDataLayersChanged.Broadcast(Actor);
+				}
+			}
+
+			if (bActorWasModified)
+			{
+				// Update per-view visibility info
+				UpdateActorAllViewsVisibility(Actor);
+
+				// Update general actor visibility
+				bool bActorModified = false;
+				bool bActorSelectionChanged = false;
+				UpdateActorVisibility(Actor, bActorSelectionChanged, bActorModified, /*bActorNotifySelectionChange*/true, /*bActorRedrawViewports*/false);
+
+				bChangesOccurred = true;
+			}
+		}
+
+		GEditor->GetSelectedActors()->EndBatchSelectOperation();
+	}
+
+	return bChangesOccurred;
+}
+
+bool UDataLayerEditorSubsystem::AddActorsToDataLayers(const TArray<TWeakObjectPtr<AActor>>& Actors, const TArray<const UDataLayer*>& DataLayers)
+{
+	TArray<AActor*> ActorsRawPtr;
+	Algo::TransformIf(Actors, ActorsRawPtr, [](const TWeakObjectPtr<AActor>& Actor) { return Actor.IsValid(); }, [](const TWeakObjectPtr<AActor>& Actor) { return Actor.Get(); });
+
+	return AddActorsToDataLayers(ActorsRawPtr, DataLayers);
+}
+
+bool UDataLayerEditorSubsystem::RemoveActorFromDataLayer(AActor* Actor, const UDataLayer* DataLayer)
+{
+	TArray<AActor*> Actors;
+	Actors.Add(Actor);
+
+	return RemoveActorsFromDataLayers(Actors, {DataLayer});
+}
+
+bool UDataLayerEditorSubsystem::RemoveActorFromDataLayers(AActor* Actor, const TArray<const UDataLayer*>& DataLayers)
+{
+	TArray<AActor*> Actors;
+	Actors.Add(Actor);
+
+	return RemoveActorsFromDataLayers(Actors, DataLayers);
+}
+
+bool UDataLayerEditorSubsystem::RemoveActorsFromDataLayer(const TArray<AActor*>& Actors, const UDataLayer* DataLayer)
+{
+	return RemoveActorsFromDataLayers(Actors, {DataLayer});
+}
+
+bool UDataLayerEditorSubsystem::RemoveActorsFromDataLayer(const TArray<TWeakObjectPtr<AActor>>& Actors, const UDataLayer* DataLayer)
+{
+	return RemoveActorsFromDataLayers(Actors, {DataLayer});
+}
+
+bool UDataLayerEditorSubsystem::RemoveActorsFromDataLayers(const TArray<AActor*>& Actors, const TArray<const UDataLayer*>& DataLayers)
+{
+	GEditor->GetSelectedActors()->BeginBatchSelectOperation();
+
+	bool bChangesOccurred = false;
+	for (AActor* Actor : Actors)
+	{
+		if (!IsActorValidForDataLayer(Actor))
+		{
+			continue;
+		}
+
+		bool bActorWasModified = false;
+		for (const UDataLayer* DataLayer : DataLayers)
+		{
+			if (Actor->RemoveDataLayer(DataLayer))
+			{
+				bActorWasModified = true;
+				DataLayerChanged.Broadcast(EDataLayerAction::Modify, DataLayer, NAME_None);
+				ActorDataLayersChanged.Broadcast(Actor);
+			}
+		}
+
+		if (bActorWasModified)
+		{
+			// Update per-view visibility info
+			UpdateActorAllViewsVisibility(Actor);
+
+			// Update general actor visibility
+			bool bActorModified = false;
+			bool bActorSelectionChanged = false;
+			UpdateActorVisibility(Actor, bActorSelectionChanged, bActorModified, /*bActorNotifySelectionChange*/true, /*bActorRedrawViewports*/false);
+
+			bChangesOccurred = true;
+		}
+	}
+
+	GEditor->GetSelectedActors()->EndBatchSelectOperation();
+
+	return bChangesOccurred;
+}
+
+bool UDataLayerEditorSubsystem::RemoveActorsFromDataLayers(const TArray<TWeakObjectPtr<AActor>>& Actors, const TArray<const UDataLayer*>& DataLayers)
+{
+	TArray<AActor*> ActorsRawPtr;
+	Algo::TransformIf(Actors, ActorsRawPtr, [](const TWeakObjectPtr<AActor>& Actor) { return Actor.IsValid(); }, [](const TWeakObjectPtr<AActor>& Actor) { return Actor.Get(); });
+
+	return RemoveActorsFromDataLayers(ActorsRawPtr, DataLayers);
+}
+
+
+///////////////////////////////////////////////////////////////////////////////////////////////////////////////////////
+//
+// Operations on selected actors.
+//
+///////////////////////////////////////////////////////////////////////////////////////////////////////////////////////
+TArray<AActor*> UDataLayerEditorSubsystem::GetSelectedActors() const
+{
+	TArray<AActor*> CurrentlySelectedActors;
+	GEditor->GetSelectedActors()->GetSelectedObjects<AActor>(CurrentlySelectedActors);
+	return CurrentlySelectedActors;
+}
+
+bool UDataLayerEditorSubsystem::AddSelectedActorsToDataLayer(const UDataLayer* DataLayer)
+{
+	return AddActorsToDataLayer(GetSelectedActors(), DataLayer);
+}
+
+bool UDataLayerEditorSubsystem::RemoveSelectedActorsFromDataLayer(const UDataLayer* DataLayer)
+{
+	return RemoveActorsFromDataLayer(GetSelectedActors(), DataLayer);
+}
+
+bool UDataLayerEditorSubsystem::AddSelectedActorsToDataLayers(const TArray<const UDataLayer*>& DataLayers)
+{
+	return AddActorsToDataLayers(GetSelectedActors(), DataLayers);
+}
+
+bool UDataLayerEditorSubsystem::RemoveSelectedActorsFromDataLayers(const TArray<const UDataLayer*>& DataLayers)
+{
+	return RemoveActorsFromDataLayers(GetSelectedActors(), DataLayers);
+}
+
+
+///////////////////////////////////////////////////////////////////////////////////////////////////////////////////////
+//
+// Operations on actors in DataLayers
+//
+///////////////////////////////////////////////////////////////////////////////////////////////////////////////////////
+bool UDataLayerEditorSubsystem::SelectActorsInDataLayer(const UDataLayer* DataLayer, const bool bSelect, const bool bNotify, const bool bSelectEvenIfHidden)
+{
+	return SelectActorsInDataLayer(DataLayer, bSelect, bNotify, bSelectEvenIfHidden, nullptr);
+}
+
+bool UDataLayerEditorSubsystem::SelectActorsInDataLayer(const UDataLayer* DataLayer, const bool bSelect, const bool bNotify, const bool bSelectEvenIfHidden, const TSharedPtr<FActorFilter>& Filter)
+{
+	bool bChangesOccurred = false;
+
+	GEditor->GetSelectedActors()->BeginBatchSelectOperation();
+	{
+		// Iterate over all actors, looking for actors in the specified DataLayers.
+		for (AActor* Actor : FActorRange(GetWorld()))
+		{
+			if (!IsActorValidForDataLayer(Actor))
+			{
+				continue;
+			}
+
+			if (Filter.IsValid() && !Filter->PassesFilter(Actor))
+			{
+				continue;
+			}
+
+			if (Actor->ContainsDataLayer(DataLayer))
+			{
+				// The actor was found to be in a specified DataLayer. Set selection state and move on to the next actor.
+				bool bNotifyForActor = false;
+				GEditor->GetSelectedActors()->Modify();
+				GEditor->SelectActor(Actor, bSelect, bNotifyForActor, bSelectEvenIfHidden);
+				bChangesOccurred = true;
+			}
+		}
+	}
+	GEditor->GetSelectedActors()->EndBatchSelectOperation();
+
+	if (bNotify)
+	{
+		GEditor->NoteSelectionChange();
+	}
+
+	return bChangesOccurred;
+}
+
+bool UDataLayerEditorSubsystem::SelectActorsInDataLayers(const TArray<const UDataLayer*>& DataLayers, const bool bSelect, const bool bNotify, const bool bSelectEvenIfHidden)
+{
+	return SelectActorsInDataLayers(DataLayers, bSelect, bNotify, bSelectEvenIfHidden, nullptr);
+}
+
+bool UDataLayerEditorSubsystem::SelectActorsInDataLayers(const TArray<const UDataLayer*>& DataLayers, const bool bSelect, const bool bNotify, const bool bSelectEvenIfHidden, const TSharedPtr<FActorFilter>& Filter)
+{
+	if (DataLayers.Num() == 0)
+	{
+		return true;
+	}
+
+	GEditor->GetSelectedActors()->BeginBatchSelectOperation();
+	bool bChangesOccurred = false;
+
+	// Iterate over all actors, looking for actors in the specified DataLayers.
+	for (AActor* Actor : FActorRange(GetWorld()))
+	{
+		if (!IsActorValidForDataLayer(Actor))
+		{
+			continue;
+		}
+
+		if (Filter.IsValid() && !Filter->PassesFilter(TWeakObjectPtr<AActor>(Actor)))
+		{
+			continue;
+		}
+
+		for (const UDataLayer* DataLayer : DataLayers)
+		{
+			if (Actor->ContainsDataLayer(DataLayer))
+			{
+				// The actor was found to be in a specified DataLayer. Set selection state and move on to the next actor.
+				bool bNotifyForActor = false;
+				GEditor->GetSelectedActors()->Modify();
+				GEditor->SelectActor(Actor, bSelect, bNotifyForActor, bSelectEvenIfHidden);
+				bChangesOccurred = true;
+				break;
+			}
+		}
+	}
+
+	GEditor->GetSelectedActors()->EndBatchSelectOperation();
+
+	if (bNotify)
+	{
+		GEditor->NoteSelectionChange();
+	}
+
+	return bChangesOccurred;
+}
+
+///////////////////////////////////////////////////////////////////////////////////////////////////////////////////////
+//
+// Operations on actor viewport visibility regarding DataLayers
+//
+///////////////////////////////////////////////////////////////////////////////////////////////////////////////////////
+void UDataLayerEditorSubsystem::UpdatePerViewVisibility(FLevelEditorViewportClient* ViewportClient, const UDataLayer* DataLayerThatChanged)
+{
+	if (!ViewportClient->GetWorld())
+	{
+		return;
+	}
+
+	// Iterate over all actors, looking for actors in the specified DataLayers.
+	const int32 ViewIndex = ViewportClient->ViewIndex;
+	for (AActor* Actor : FActorRange(ViewportClient->GetWorld()))
+	{
+		if (!IsActorValidForDataLayer(Actor))
+		{
+			continue;
+		}
+
+		// If the view has nothing hidden, just quickly mark the actor as visible in this view 
+		if (ViewportClient->ViewHiddenDataLayers.Num() == 0)
+		{
+			// If the actor had this view hidden, then unhide it
+			if (Actor->HiddenEditorViews & ((uint64)1 << ViewIndex))
+			{
+				// Make sure this actor doesn't have the view set
+				Actor->HiddenEditorViews &= ~((uint64)1 << ViewIndex);
+				Actor->MarkComponentsRenderStateDirty();
+			}
+		}
+		// Else if we were given a name that was changed, only update actors with that name in their DataLayers, otherwise update all actors
+		else if (DataLayerThatChanged == nullptr || Actor->ContainsDataLayer(DataLayerThatChanged))
+		{
+			UpdateActorViewVisibility(ViewportClient, Actor);
+		}
+	}
+
+	// Make sure we redraw the viewport
+	ViewportClient->Invalidate();
+}
+
+void UDataLayerEditorSubsystem::UpdateAllViewVisibility(const UDataLayer* DataLayerThatChanged)
+{
+	// Update all view's hidden DataLayers if they had this one
+	for (FLevelEditorViewportClient* ViewportClient : GEditor->GetLevelViewportClients())
+	{
+		UpdatePerViewVisibility(ViewportClient, DataLayerThatChanged);
+	}
+}
+
+void UDataLayerEditorSubsystem::UpdateActorViewVisibility(FLevelEditorViewportClient* ViewportClient, AActor* Actor, bool bReregisterIfDirty)
+{
+	const int32 ViewIndex = ViewportClient->ViewIndex;
+	uint64 OriginalHiddenViews = Actor->HiddenEditorViews;
+
+	// Update Actor's HiddenEditorViews to reflect ViewHiddenDataLayers
+	if (bool bIsHiddenByViewHiddenDataLayers = Actor->HasAnyOfDataLayers(ViewportClient->ViewHiddenDataLayers))
+	{
+		Actor->HiddenEditorViews |= ((uint64)1 << ViewIndex);
+	}
+	else
+	{
+		Actor->HiddenEditorViews &= ~((uint64)1 << ViewIndex);
+	}
+
+	// Re-register if we changed the visibility bits, as the rendering thread needs them
+	if (bReregisterIfDirty && OriginalHiddenViews != Actor->HiddenEditorViews)
+	{
+		Actor->MarkComponentsRenderStateDirty();
+
+		// Make sure we redraw the viewport
+		ViewportClient->Invalidate();
+	}
+}
+
+void UDataLayerEditorSubsystem::UpdateActorAllViewsVisibility(AActor* Actor)
+{
+	uint64 OriginalHiddenViews = Actor->HiddenEditorViews;
+
+	for (FLevelEditorViewportClient* ViewportClient : GEditor->GetLevelViewportClients())
+	{
+		// Don't have this reattach, as we can do it once for all views
+		UpdateActorViewVisibility(ViewportClient, Actor, false);
+	}
+
+	// Re-register if we changed the visibility bits, as the rendering thread needs them
+	if (OriginalHiddenViews != Actor->HiddenEditorViews)
+	{
+		return;
+	}
+
+	Actor->MarkComponentsRenderStateDirty();
+
+	// Redraw all viewports
+	for (FLevelEditorViewportClient* ViewportClient : GEditor->GetLevelViewportClients())
+	{
+		ViewportClient->Invalidate();
+	}
+}
+
+void UDataLayerEditorSubsystem::RemoveViewFromActorViewVisibility(FLevelEditorViewportClient* ViewportClient)
+{
+	if (!ViewportClient->GetWorld())
+	{
+		return;
+	}
+
+	// Get the bit for the view index
+	const int32 ViewIndex = ViewportClient->ViewIndex;
+	uint64 ViewBit = ((uint64)1 << ViewIndex);
+	// Get all bits under that that we want to keep
+	uint64 KeepBits = ViewBit - 1;
+
+	// Iterate over all actors, looking for actors in the specified DataLayers.
+	for (AActor* Actor : FActorRange(ViewportClient->GetWorld()))
+	{
+		if (!IsActorValidForDataLayer(Actor))
+		{
+			continue;
+		}
+
+		// Remember original bits
+		uint64 OriginalHiddenViews = Actor->HiddenEditorViews;
+		uint64 Was = Actor->HiddenEditorViews;
+
+		// Slide all bits higher than ViewIndex down one since the view is being removed from Editor
+		uint64 LowBits = Actor->HiddenEditorViews & KeepBits;
+
+		// Now slide the top bits down by ViewIndex + 1 (chopping off ViewBit)
+		uint64 HighBits = Actor->HiddenEditorViews >> (ViewIndex + 1);
+		// Then slide back up by ViewIndex, which will now have erased ViewBit, as well as leaving 0 in the low bits
+		HighBits = HighBits << ViewIndex;
+
+		// Put it all back together
+		Actor->HiddenEditorViews = LowBits | HighBits;
+
+		// Re-register if we changed the visibility bits, as the rendering thread needs them
+		if (OriginalHiddenViews == Actor->HiddenEditorViews)
+		{
+			continue;
+		}
+
+		// Find all registered primitive components and update the scene proxy with the actors updated visibility map
+		TInlineComponentArray<UPrimitiveComponent*> Components;
+		Actor->GetComponents(Components);
+
+		for (UActorComponent* Component : Actor->GetComponents())
+		{
+			UPrimitiveComponent* PrimitiveComponent = Cast<UPrimitiveComponent>(Component);
+			if (PrimitiveComponent && PrimitiveComponent->IsRegistered())
+			{
+				// Push visibility to the render thread
+				PrimitiveComponent->PushEditorVisibilityToProxy(Actor->HiddenEditorViews);
+			}
+		}
+	}
+}
+
+static void UpdateBrushDataLayerVisibility(ABrush* Brush, bool bIsHidden)
+{
+	ULevel* Level = Brush->GetLevel();
+	if (!Level)
+	{
+		return;
+	}
+
+	UModel* Model = Level->Model;
+	if (!Model)
+	{
+		return;
+	}
+
+	bool bAnySurfaceWasFound = false;
+	for (FBspSurf& Surf : Model->Surfs)
+	{
+		if (Surf.Actor == Brush)
+		{
+			Surf.bHiddenEdLayer = bIsHidden;
+			bAnySurfaceWasFound = true;
+		}
+	}
+
+	if (bAnySurfaceWasFound)
+	{
+		Level->UpdateModelComponents();
+		Model->InvalidSurfaces = true;
+	}
+}
+
+bool UDataLayerEditorSubsystem::UpdateActorVisibility(AActor* Actor, bool& bOutSelectionChanged, bool& bOutActorModified, const bool bNotifySelectionChange, const bool bRedrawViewports)
+{
+	bOutActorModified = false;
+	bOutSelectionChanged = false;
+
+	if (!IsActorValidForDataLayer(Actor))
+	{
+		return false;
+	}
+
+	// If the actor doesn't belong to any DataLayers
+	if (!Actor->HasValidDataLayers())
+	{
+		// If the actor is also hidden
+		if (Actor->bHiddenEdLayer)
+		{
+			// Actors that don't belong to any DataLayer shouldn't be hidden
+			Actor->bHiddenEdLayer = false;
+			Actor->MarkComponentsRenderStateDirty();
+			bOutActorModified = true;
+		}
+
+		return bOutActorModified;
+	}
+
+	bool bActorBelongsToVisibleDataLayer = false;
+	if (AWorldDataLayers* WorldDataLayers = GetWorldDataLayers())
+	{
+		WorldDataLayers->ForEachDataLayer([&Actor, &bOutActorModified, &bActorBelongsToVisibleDataLayer](UDataLayer* DataLayer)
+		{
+			if (DataLayer->IsVisible() && Actor->ContainsDataLayer(DataLayer))
+			{
+				if (Actor->bHiddenEdLayer)
+				{
+					Actor->bHiddenEdLayer = false;
+					Actor->MarkComponentsRenderStateDirty();
+					bOutActorModified = true;
+
+					if (ABrush* Brush = Cast<ABrush>(Actor))
+					{
+						UpdateBrushDataLayerVisibility(Brush, /*bIsHidden*/false);
+					}
+				}
+
+				// Stop, because we found at least one visible DataLayer the actor belongs to
+				bActorBelongsToVisibleDataLayer = true;
+				return false;
+			}
+			return true;
+		});
+	}
+
+	// If the actor isn't part of a visible DataLayer, hide and de-select it.
+	if (!bActorBelongsToVisibleDataLayer)
+	{
+		if (!Actor->bHiddenEdLayer)
+		{
+			Actor->bHiddenEdLayer = true;
+			Actor->MarkComponentsRenderStateDirty();
+			bOutActorModified = true;
+
+			if (ABrush* Brush = Cast<ABrush>(Actor))
+			{
+				UpdateBrushDataLayerVisibility(Brush, /*bIsHidden*/false);
+			}
+		}
+
+		// If the actor was selected, mark it as unselected
+		if (Actor->IsSelected())
+		{
+			bool bSelect = false;
+			bool bNotify = false;
+			bool bIncludeHidden = true;
+			GEditor->SelectActor(Actor, bSelect, bNotify, bIncludeHidden);
+
+			bOutSelectionChanged = true;
+			bOutActorModified = true;
+		}
+	}
+
+	if (bNotifySelectionChange && bOutSelectionChanged)
+	{
+		GEditor->NoteSelectionChange();
+	}
+
+	if (bRedrawViewports)
+	{
+		GEditor->RedrawLevelEditingViewports();
+	}
+
+	return bOutActorModified || bOutSelectionChanged;
+}
+
+bool UDataLayerEditorSubsystem::UpdateAllActorsVisibility(const bool bNotifySelectionChange, const bool bRedrawViewports)
+{
+	bool bSelectionChanged = false;
+	bool bChangesOccurred = false;
+	for (AActor* Actor : FActorRange(GetWorld()))
+	{
+		bool bActorModified = false;
+		bool bActorSelectionChanged = false;
+		bChangesOccurred |= UpdateActorVisibility(Actor, bActorSelectionChanged, bActorModified, /*bActorNotifySelectionChange*/false, /*bActorRedrawViewports*/false);
+		bSelectionChanged |= bActorSelectionChanged;
+	}
+
+	if (bNotifySelectionChange && bSelectionChanged)
+	{
+		GEditor->NoteSelectionChange();
+	}
+
+	if (bRedrawViewports)
+	{
+		GEditor->RedrawLevelEditingViewports();
+	}
+
+	return bChangesOccurred;
+}
+
+
+///////////////////////////////////////////////////////////////////////////////////////////////////////////////////////
+//
+// Operations on DataLayers
+//
+///////////////////////////////////////////////////////////////////////////////////////////////////////////////////////
+
+void UDataLayerEditorSubsystem::AppendActorsFromDataLayer(const UDataLayer* DataLayer, TArray<AActor*>& InOutActors) const
+{
+	AppendActorsFromDataLayer(DataLayer, InOutActors, nullptr);
+}
+
+void UDataLayerEditorSubsystem::AppendActorsFromDataLayer(const UDataLayer* DataLayer, TArray<AActor*>& InOutActors, const TSharedPtr<FActorFilter>& Filter) const
+{
+	for (AActor* Actor : FActorRange(GetWorld()))
+	{
+		if (Filter.IsValid() && !Filter->PassesFilter(Actor))
+		{
+			continue;
+		}
+		if (Actor->ContainsDataLayer(DataLayer))
+		{
+			InOutActors.Add(Actor);
+		}
+	}
+}
+
+void UDataLayerEditorSubsystem::AppendActorsFromDataLayer(const UDataLayer* DataLayer, TArray<TWeakObjectPtr<AActor>>& InOutActors, const TSharedPtr<FActorFilter>& Filter) const
+{
+	for (AActor* Actor : FActorRange(GetWorld()))
+	{
+		if (Filter.IsValid() && !Filter->PassesFilter(Actor))
+		{
+			continue;
+		}
+		if (Actor->ContainsDataLayer(DataLayer))
+		{
+			InOutActors.Add(Actor);
+		}
+	}
+}
+
+void UDataLayerEditorSubsystem::AppendActorsFromDataLayers(const TArray<const UDataLayer*>& DataLayers, TArray<AActor*>& InOutActors) const
+{
+	AppendActorsFromDataLayers(DataLayers, InOutActors, nullptr);
+}
+
+void UDataLayerEditorSubsystem::AppendActorsFromDataLayers(const TArray<const UDataLayer*>& DataLayers, TArray<AActor*>& InOutActors, const TSharedPtr<FActorFilter>& Filter) const
+{
+	for (AActor* Actor : FActorRange(GetWorld()))
+	{
+		if (Filter.IsValid() && !Filter->PassesFilter(Actor))
+		{
+			continue;
+		}
+		for (const UDataLayer* DataLayer : DataLayers)
+		{
+			if (Actor->ContainsDataLayer(DataLayer))
+			{
+				InOutActors.Add(Actor);
+				break;
+			}
+		}
+	}
+}
+
+void UDataLayerEditorSubsystem::AppendActorsFromDataLayers(const TArray<const UDataLayer*>& DataLayers, TArray<TWeakObjectPtr<AActor>>& InOutActors, const TSharedPtr<FActorFilter>& Filter) const
+{
+	for (AActor* Actor : FActorRange(GetWorld()))
+	{
+		if (Filter.IsValid() && !Filter->PassesFilter(Actor))
+		{
+			continue;
+		}
+		for (const UDataLayer* DataLayer : DataLayers)
+		{
+			if (Actor->ContainsDataLayer(DataLayer))
+			{
+				InOutActors.Add(Actor);
+				break;
+			}
+		}
+	}
+}
+
+TArray<AActor*> UDataLayerEditorSubsystem::GetActorsFromDataLayer(const UDataLayer* DataLayer) const
+{
+	TArray<AActor*> OutActors;
+	AppendActorsFromDataLayer(DataLayer, OutActors);
+	return OutActors;
+}
+
+TArray<AActor*> UDataLayerEditorSubsystem::GetActorsFromDataLayer(const UDataLayer* DataLayer, const TSharedPtr<FActorFilter>& Filter) const
+{
+	TArray<AActor*> OutActors;
+	AppendActorsFromDataLayer(DataLayer, OutActors, Filter);
+	return OutActors;
+}
+
+TArray<AActor*> UDataLayerEditorSubsystem::GetActorsFromDataLayers(const TArray<const UDataLayer*>& DataLayers) const
+{
+	TArray<AActor*> OutActors;
+	AppendActorsFromDataLayers(DataLayers, OutActors);
+	return OutActors;
+}
+
+TArray<AActor*> UDataLayerEditorSubsystem::GetActorsFromDataLayers(const TArray<const UDataLayer*>& DataLayers, const TSharedPtr<FActorFilter>& Filter) const
+{
+	TArray<AActor*> OutActors;
+	AppendActorsFromDataLayers(DataLayers, OutActors, Filter);
+	return OutActors;
+}
+
+void UDataLayerEditorSubsystem::SetDataLayerVisibility(UDataLayer* DataLayer, const bool bIsVisible)
+{
+	SetDataLayersVisibility({ DataLayer }, bIsVisible);
+}
+
+void UDataLayerEditorSubsystem::SetDataLayersVisibility(const TArray<UDataLayer*>& DataLayers, const bool bIsVisible)
+{
+	bool bChangeOccurred = false;
+	for (UDataLayer* DataLayer : DataLayers)
+	{
+		check(DataLayer);
+
+		if (DataLayer->IsVisible() != bIsVisible)
+		{
+			DataLayer->Modify();
+			DataLayer->SetVisible(bIsVisible);
+			DataLayerChanged.Broadcast(EDataLayerAction::Modify, DataLayer, "bIsVisible");
+			bChangeOccurred = true;
+		}
+	}
+
+	if (bChangeOccurred)
+	{
+		UpdateAllActorsVisibility(true, true);
+	}
+}
+
+void UDataLayerEditorSubsystem::ToggleDataLayerVisibility(UDataLayer* DataLayer)
+{
+	check(DataLayer);
+
+	DataLayer->Modify();
+	DataLayer->SetVisible(!DataLayer->IsVisible());
+
+	DataLayerChanged.Broadcast(EDataLayerAction::Modify, DataLayer, "bIsVisible");
+	UpdateAllActorsVisibility(true, true);
+}
+
+void UDataLayerEditorSubsystem::ToggleDataLayersVisibility(const TArray<UDataLayer*>& DataLayers)
+{
+	if (DataLayers.Num() == 0)
+	{
+		return;
+	}
+
+	for (UDataLayer* DataLayer : DataLayers)
+	{
+		DataLayer->Modify();
+		DataLayer->SetVisible(!DataLayer->IsVisible());
+		DataLayerChanged.Broadcast(EDataLayerAction::Modify, DataLayer, "bIsVisible");
+	}
+
+	UpdateAllActorsVisibility(true, true);
+}
+
+void UDataLayerEditorSubsystem::MakeAllDataLayersVisible()
+{
+	if (AWorldDataLayers* WorldDataLayers = GetWorldDataLayers())
+	{
+		WorldDataLayers->ForEachDataLayer([this](UDataLayer* DataLayer)
+		{
+			if (!DataLayer->IsVisible())
+			{
+				DataLayer->Modify();
+				DataLayer->SetVisible(true);
+				DataLayerChanged.Broadcast(EDataLayerAction::Modify, DataLayer, "bIsVisible");
+			}
+			return true;
+		});
+	}
+	
+	UpdateAllActorsVisibility(true, true);
+}
+
+void UDataLayerEditorSubsystem::SetDataLayerIsDynamicallyLoaded(UDataLayer* DataLayer, const bool bIsDynamicallyLoaded)
+{
+	SetDataLayersIsDynamicallyLoaded({ DataLayer }, bIsDynamicallyLoaded);
+}
+
+void UDataLayerEditorSubsystem::SetDataLayersIsDynamicallyLoaded(const TArray<UDataLayer*>& DataLayers, const bool bIsDynamicallyLoaded)
+{
+	bool bChangeOccurred = false;
+	for (UDataLayer* DataLayer : DataLayers)
+	{
+		if (DataLayer->IsDynamicallyLoaded() != bIsDynamicallyLoaded)
+		{
+			DataLayer->Modify();
+			DataLayer->SetIsDynamicallyLoaded(bIsDynamicallyLoaded);
+			DataLayerChanged.Broadcast(EDataLayerAction::Modify, DataLayer, "bIsDynamicallyLoaded");
+			bChangeOccurred = true;
+		}
+	}
+
+	if (bChangeOccurred)
+	{
+		GEditor->RedrawLevelEditingViewports();
+	}
+}
+
+void UDataLayerEditorSubsystem::ToggleDataLayerIsDynamicallyLoaded(UDataLayer* DataLayer)
+{
+	check(DataLayer);
+
+	DataLayer->Modify();
+	DataLayer->SetIsDynamicallyLoaded(!DataLayer->IsDynamicallyLoaded());
+
+	DataLayerChanged.Broadcast(EDataLayerAction::Modify, DataLayer, "bIsDynamicallyLoaded");
+	GEditor->RedrawLevelEditingViewports();
+}
+
+void UDataLayerEditorSubsystem::ToggleDataLayersIsDynamicallyLoaded(const TArray<UDataLayer*>& DataLayers)
+{
+	if (DataLayers.Num() == 0)
+	{
+		return;
+	}
+
+	for (UDataLayer* DataLayer : DataLayers)
+	{
+		DataLayer->Modify();
+		DataLayer->SetIsDynamicallyLoaded(!DataLayer->IsDynamicallyLoaded());
+		DataLayerChanged.Broadcast(EDataLayerAction::Modify, DataLayer, "bIsDynamicallyLoaded");
+	}
+
+	GEditor->RedrawLevelEditingViewports();
+}
+
+const UDataLayer* UDataLayerEditorSubsystem::GetDataLayerFromName(const FName& DataLayerName) const
+{
+	const AWorldDataLayers* WorldDataLayers = GetWorldDataLayers();
+	return WorldDataLayers ? WorldDataLayers->GetDataLayerFromName(DataLayerName) : nullptr;
+}
+
+const UDataLayer* UDataLayerEditorSubsystem::GetDataLayerFromLabel(const FName& DataLayerLabel) const
+{
+	const AWorldDataLayers* WorldDataLayers = GetWorldDataLayers();
+	return WorldDataLayers ? WorldDataLayers->GetDataLayerFromLabel(DataLayerLabel) : nullptr;
+}
+
+bool UDataLayerEditorSubsystem::TryGetDataLayerFromLabel(const FName& DataLayerLabel, const UDataLayer*& OutDataLayer)
+{
+	OutDataLayer = GetDataLayerFromLabel(DataLayerLabel);
+	return (OutDataLayer != nullptr);
+}
+
+const AWorldDataLayers* UDataLayerEditorSubsystem::GetWorldDataLayers() const
+{
+	return AWorldDataLayers::Get(GetWorld());
+}
+
+AWorldDataLayers* UDataLayerEditorSubsystem::GetWorldDataLayers(bool bCreateIfNotFound)
+{
+	UWorld* World = GetWorld();
+	FName WorldDataLayersName = AWorldDataLayers::StaticClass()->GetFName();
+	AWorldDataLayers* WorldDataLayers = nullptr;
+	if (UObject* ExistingObject = StaticFindObject(nullptr, World->PersistentLevel, *WorldDataLayersName.ToString()))
+	{
+		WorldDataLayers = CastChecked<AWorldDataLayers>(ExistingObject);
+		if (WorldDataLayers->IsPendingKill())
+		{
+			// Handle the case where the actor already exists, but it's pending kill
+			WorldDataLayers->Rename(nullptr, nullptr, REN_DontCreateRedirectors | REN_DoNotDirty | REN_NonTransactional | REN_ForceNoResetLoaders);
+			WorldDataLayers = nullptr;
+		}
+	}
+
+	if (!WorldDataLayers && bCreateIfNotFound)
+	{
+		FActorSpawnParameters SpawnParams;
+		SpawnParams.OverrideLevel = World->PersistentLevel;
+		SpawnParams.bHideFromSceneOutliner = true;
+		SpawnParams.Name = WorldDataLayersName;
+		SpawnParams.NameMode = FActorSpawnParameters::ESpawnActorNameMode::Required_Fatal;
+		WorldDataLayers = World->SpawnActor<AWorldDataLayers>(AWorldDataLayers::StaticClass(), SpawnParams);
+	}
+
+	return WorldDataLayers;
+}
+
+void UDataLayerEditorSubsystem::AddAllDataLayersTo(TArray<TWeakObjectPtr<UDataLayer>>& OutDataLayers) const
+{
+	if (const AWorldDataLayers* WorldDataLayers = GetWorldDataLayers())
+	{
+		WorldDataLayers->ForEachDataLayer([&OutDataLayers](UDataLayer* DataLayer)
+		{
+			OutDataLayers.Add(DataLayer);
+			return true;
+		});
+	}
+}
+
+UDataLayer* UDataLayerEditorSubsystem::CreateDataLayer()
+{
+	AWorldDataLayers* WorldDataLayers = GetWorldDataLayers(/*bCreateIfNotFound*/true);
+	UDataLayer* NewDataLayer = WorldDataLayers->CreateDataLayer();
+	DataLayerChanged.Broadcast(EDataLayerAction::Add, NewDataLayer, NAME_None);
+	return NewDataLayer;
+}
+
+void UDataLayerEditorSubsystem::DeleteDataLayers(const TArray<const UDataLayer*>& DataLayersToDelete)
+{
+	if (AWorldDataLayers* WorldDataLayers = GetWorldDataLayers())
+	{
+		if (WorldDataLayers->RemoveDataLayers(DataLayersToDelete))
+		{
+			DataLayerChanged.Broadcast(EDataLayerAction::Delete, NULL, NAME_None);
+		}
+	}
+}
+
+void UDataLayerEditorSubsystem::DeleteDataLayer(const UDataLayer* DataLayerToDelete)
+{
+	if (AWorldDataLayers* WorldDataLayers = GetWorldDataLayers())
+	{
+		if (WorldDataLayers->RemoveDataLayer(DataLayerToDelete))
+		{
+			DataLayerChanged.Broadcast(EDataLayerAction::Delete, NULL, NAME_None);
+		}
+	}
+}
+
+bool UDataLayerEditorSubsystem::RenameDataLayer(UDataLayer* DataLayer, const FName& DataLayerLabel)
+{
+	if (DataLayer->GetDataLayerLabel() != DataLayerLabel)
+	{
+		if (AWorldDataLayers* WorldDataLayers = GetWorldDataLayers())
+		{
+			FName UniqueNewDataLayerLabel = WorldDataLayers->GenerateUniqueDataLayerLabel(DataLayerLabel);
+			DataLayer->Modify();
+			DataLayer->SetDataLayerLabel(UniqueNewDataLayerLabel);
+			DataLayerChanged.Broadcast(EDataLayerAction::Rename, DataLayer, "DataLayerLabel");
+			return true;
+		}
+	}
+	return false;
+}
\ No newline at end of file
diff --git a/Engine/Source/Editor/DataLayerEditor/Private/DataLayer/DataLayerHierarchy.cpp b/Engine/Source/Editor/DataLayerEditor/Private/DataLayer/DataLayerHierarchy.cpp
new file mode 100644
--- /dev/null
+++ b/Engine/Source/Editor/DataLayerEditor/Private/DataLayer/DataLayerHierarchy.cpp
@@ -0,0 +1,219 @@
+// Copyright Epic Games, Inc. All Rights Reserved.
+
+#include "DataLayerHierarchy.h"
+#include "DataLayerMode.h"
+#include "DataLayerActorTreeItem.h"
+#include "DataLayerTreeItem.h"
+#include "SDataLayerBrowser.h"
+#include "DataLayerMode.h"
+#include "DataLayer/DataLayerEditorSubsystem.h"
+#include "WorldPartition/DataLayer/DataLayer.h"
+#include "Engine/World.h"
+#include "EngineUtils.h"
+#include "Editor.h"
+
+TUniquePtr<FDataLayerHierarchy> FDataLayerHierarchy::Create(FDataLayerMode* Mode, const TWeakObjectPtr<UWorld>& World)
+{
+	FDataLayerHierarchy* Hierarchy = new FDataLayerHierarchy(Mode, World);
+
+	GEngine->OnLevelActorAdded().AddRaw(Hierarchy, &FDataLayerHierarchy::OnLevelActorAdded);
+	GEngine->OnLevelActorDeleted().AddRaw(Hierarchy, &FDataLayerHierarchy::OnLevelActorDeleted);
+	GEngine->OnLevelActorListChanged().AddRaw(Hierarchy, &FDataLayerHierarchy::OnLevelActorListChanged);
+	UDataLayerEditorSubsystem::Get()->OnDataLayerChanged().AddRaw(Hierarchy, &FDataLayerHierarchy::OnDataLayerChanged);
+	UDataLayerEditorSubsystem::Get()->OnActorDataLayersChanged().AddRaw(Hierarchy, &FDataLayerHierarchy::OnActorDataLayersChanged);
+	Mode->GetDataLayerBrowser()->OnModeChanged().AddRaw(Hierarchy, &FDataLayerHierarchy::OnDataLayerBrowserModeChanged);
+
+	FWorldDelegates::LevelAddedToWorld.AddRaw(Hierarchy, &FDataLayerHierarchy::OnLevelAdded);
+	FWorldDelegates::LevelRemovedFromWorld.AddRaw(Hierarchy, &FDataLayerHierarchy::OnLevelRemoved);
+
+	return TUniquePtr<FDataLayerHierarchy>(Hierarchy);
+}
+
+FDataLayerHierarchy::FDataLayerHierarchy(FDataLayerMode* Mode, const TWeakObjectPtr<UWorld>& World)
+	: ISceneOutlinerHierarchy(Mode)
+	, RepresentingWorld(World)
+{
+}
+
+FDataLayerHierarchy::~FDataLayerHierarchy()
+{
+	if (GEngine)
+	{
+		GEngine->OnLevelActorAdded().RemoveAll(this);
+		GEngine->OnLevelActorDeleted().RemoveAll(this);
+		GEngine->OnLevelActorListChanged().RemoveAll(this);
+	}
+
+	UDataLayerEditorSubsystem::Get()->OnDataLayerChanged().RemoveAll(this);
+	UDataLayerEditorSubsystem::Get()->OnActorDataLayersChanged().RemoveAll(this);
+
+	if (SDataLayerBrowser* DataLayerBrowser = GetDataLayerMode() ? GetDataLayerMode()->GetDataLayerBrowser() : nullptr)
+	{
+		DataLayerBrowser->OnModeChanged().RemoveAll(this);
+	}
+
+	FWorldDelegates::LevelAddedToWorld.RemoveAll(this);
+	FWorldDelegates::LevelRemovedFromWorld.RemoveAll(this);
+}
+
+FDataLayerMode* FDataLayerHierarchy::GetDataLayerMode() const
+{
+	return (FDataLayerMode*)Mode;
+}
+
+void FDataLayerHierarchy::CreateItems(TArray<FSceneOutlinerTreeItemPtr>& OutItems) const
+{
+	const AWorldDataLayers* WorldDataLayers = AWorldDataLayers::Get(RepresentingWorld.Get());
+	if (!WorldDataLayers)
+	{
+		return;
+	}
+
+	WorldDataLayers->ForEachDataLayer([this, &OutItems](UDataLayer* DataLayer)
+		{
+			if (FSceneOutlinerTreeItemPtr DataLayerItem = Mode->CreateItemFor<FDataLayerTreeItem>(DataLayer))
+			{
+				OutItems.Add(DataLayerItem);
+			}
+			return true;
+		});
+
+	if (GetDataLayerMode()->GetDataLayerBrowser()->GetMode() == EDataLayerBrowserMode::DataLayerContents)
+	{
+		for (AActor* Actor : FActorRange(RepresentingWorld.Get()))
+		{
+			if (Actor->HasDataLayers())
+			{
+				for (const UDataLayer* DataLayer : Actor->GetDataLayerObjects())
+				{
+					if (FSceneOutlinerTreeItemPtr DataLayerActorItem = Mode->CreateItemFor<FDataLayerActorTreeItem>(FDataLayerActorTreeItemData(Actor, const_cast<UDataLayer*>(DataLayer))))
+					{
+						OutItems.Add(DataLayerActorItem);
+					}
+				}
+			}
+		}
+	}
+}
+
+FSceneOutlinerTreeItemPtr FDataLayerHierarchy::FindParent(const ISceneOutlinerTreeItem& Item, const TMap<FSceneOutlinerTreeItemID, FSceneOutlinerTreeItemPtr>& Items) const
+{
+	if (Item.IsA<FDataLayerTreeItem>())
+	{
+		return nullptr;
+	}
+	else if (const FDataLayerActorTreeItem* DataLayerActorTreeItem = Item.CastTo<FDataLayerActorTreeItem>())
+	{
+		if (const UDataLayer* DataLayer = DataLayerActorTreeItem->GetDataLayer())
+		{
+			if (const FSceneOutlinerTreeItemPtr* ParentItem = Items.Find(DataLayer))
+			{
+				return *ParentItem;
+			}
+		}
+	}
+	return nullptr;
+}
+
+FSceneOutlinerTreeItemPtr FDataLayerHierarchy::CreateParentItem(const FSceneOutlinerTreeItemPtr& Item) const
+{
+	if (Item->IsA<FDataLayerTreeItem>())
+	{
+		return nullptr;
+	}
+	else if (FDataLayerActorTreeItem* DataLayerActorTreeItem = Item->CastTo<FDataLayerActorTreeItem>())
+	{
+		if (UDataLayer* DataLayer = DataLayerActorTreeItem->GetDataLayer())
+		{
+			return Mode->CreateItemFor<FDataLayerTreeItem>(DataLayer);
+		}
+	}
+	return nullptr;
+}
+
+void FDataLayerHierarchy::OnLevelActorAdded(AActor* InActor)
+{
+	if (InActor && RepresentingWorld.Get() == InActor->GetWorld())
+	{
+		if (InActor->HasDataLayers())
+		{
+			FSceneOutlinerHierarchyChangedData EventData;
+			EventData.Type = FSceneOutlinerHierarchyChangedData::Added;
+
+			for (const UDataLayer* DataLayer : InActor->GetDataLayerObjects())
+			{
+				EventData.Item = Mode->CreateItemFor<FDataLayerActorTreeItem>(FDataLayerActorTreeItemData(InActor, const_cast<UDataLayer*>(DataLayer)));
+				HierarchyChangedEvent.Broadcast(EventData);
+			}
+		}
+	}
+}
+
+void FDataLayerHierarchy::OnActorDataLayersChanged(const TWeakObjectPtr<AActor>& InActor)
+{
+	AActor* Actor = InActor.Get();
+	if (Actor && RepresentingWorld.Get() == Actor->GetWorld())
+	{
+		FullRefreshEvent();
+	}
+}
+
+void FDataLayerHierarchy::OnDataLayerChanged(const EDataLayerAction Action, const TWeakObjectPtr<const UDataLayer>& ChangedDataLayer, const FName& ChangedProperty)
+{
+	const UDataLayer* DataLayer = ChangedDataLayer.Get();
+	if (DataLayer && (RepresentingWorld.Get() == DataLayer->GetWorld()) || Action == EDataLayerAction::Delete || Action == EDataLayerAction::Reset)
+	{
+		FullRefreshEvent();
+	}
+}
+
+void FDataLayerHierarchy::OnDataLayerBrowserModeChanged(EDataLayerBrowserMode InMode)
+{
+	FullRefreshEvent();
+}
+
+void FDataLayerHierarchy::OnLevelActorDeleted(AActor* InActor)
+{
+	if (RepresentingWorld.Get() == InActor->GetWorld())
+	{
+		if (InActor->HasDataLayers())
+		{
+			FSceneOutlinerHierarchyChangedData EventData;
+			EventData.Type = FSceneOutlinerHierarchyChangedData::Removed;
+
+			for (const UDataLayer* DataLayer : InActor->GetDataLayerObjects())
+			{
+				EventData.ItemID = FDataLayerActorTreeItem::ComputeTreeItemID(InActor, DataLayer);
+				HierarchyChangedEvent.Broadcast(EventData);
+			}
+		}
+	}
+}
+
+void FDataLayerHierarchy::OnLevelActorListChanged()
+{
+	FullRefreshEvent();
+}
+
+void FDataLayerHierarchy::OnLevelAdded(ULevel* InLevel, UWorld* InWorld)
+{
+	if (RepresentingWorld.Get() == InWorld)
+	{
+		FullRefreshEvent();
+	}
+}
+
+void FDataLayerHierarchy::OnLevelRemoved(ULevel* InLevel, UWorld* InWorld)
+{
+	if (RepresentingWorld.Get() == InWorld)
+	{
+		FullRefreshEvent();
+	}
+}
+
+void FDataLayerHierarchy::FullRefreshEvent()
+{
+	FSceneOutlinerHierarchyChangedData EventData;
+	EventData.Type = FSceneOutlinerHierarchyChangedData::FullRefresh;
+	HierarchyChangedEvent.Broadcast(EventData);
+}
\ No newline at end of file
diff --git a/Engine/Source/Editor/DataLayerEditor/Private/DataLayer/DataLayerHierarchy.h b/Engine/Source/Editor/DataLayerEditor/Private/DataLayer/DataLayerHierarchy.h
new file mode 100644
--- /dev/null
+++ b/Engine/Source/Editor/DataLayerEditor/Private/DataLayer/DataLayerHierarchy.h
@@ -0,0 +1,42 @@
+// Copyright Epic Games, Inc. All Rights Reserved.
+
+#pragma once
+
+#include "CoreMinimal.h"
+#include "SceneOutlinerFwd.h"
+#include "SDataLayerBrowser.h"
+#include "ISceneOutlinerHierarchy.h"
+#include "DataLayer/DataLayerAction.h"
+
+class FDataLayerMode;
+class UDataLayer;
+class UWorld;
+
+class FDataLayerHierarchy : public ISceneOutlinerHierarchy
+{
+public:
+	virtual ~FDataLayerHierarchy();
+	static TUniquePtr<FDataLayerHierarchy> Create(FDataLayerMode* Mode, const TWeakObjectPtr<UWorld>& World);
+	virtual FSceneOutlinerTreeItemPtr FindParent(const ISceneOutlinerTreeItem& Item, const TMap<FSceneOutlinerTreeItemID, FSceneOutlinerTreeItemPtr>& Items) const override;
+	virtual void CreateItems(TArray<FSceneOutlinerTreeItemPtr>& OutItems) const override;
+	virtual void CreateChildren(const FSceneOutlinerTreeItemPtr& Item, TArray<FSceneOutlinerTreeItemPtr>& OutChildren) const override {}
+	virtual FSceneOutlinerTreeItemPtr CreateParentItem(const FSceneOutlinerTreeItemPtr& Item) const override;
+
+private:
+	FDataLayerHierarchy(FDataLayerMode* Mode, const TWeakObjectPtr<UWorld>& Worlds);
+	FDataLayerHierarchy(const FDataLayerHierarchy&) = delete;
+	FDataLayerHierarchy& operator=(const FDataLayerHierarchy&) = delete;
+
+	void OnLevelActorAdded(AActor* InActor);
+	void OnLevelActorDeleted(AActor* InActor);
+	void OnLevelActorListChanged();
+	void OnLevelAdded(ULevel* InLevel, UWorld* InWorld);
+	void OnLevelRemoved(ULevel* InLevel, UWorld* InWorld);
+	void OnActorDataLayersChanged(const TWeakObjectPtr<AActor>& InActor);
+	void OnDataLayerChanged(const EDataLayerAction Action, const TWeakObjectPtr<const UDataLayer>& ChangedDataLayer, const FName& ChangedProperty);
+	void OnDataLayerBrowserModeChanged(EDataLayerBrowserMode InMode);
+	void FullRefreshEvent();
+	FDataLayerMode* GetDataLayerMode() const;
+
+	TWeakObjectPtr<UWorld> RepresentingWorld;
+};
\ No newline at end of file
diff --git a/Engine/Source/Editor/DataLayerEditor/Private/DataLayer/DataLayerMode.cpp b/Engine/Source/Editor/DataLayerEditor/Private/DataLayer/DataLayerMode.cpp
new file mode 100644
--- /dev/null
+++ b/Engine/Source/Editor/DataLayerEditor/Private/DataLayer/DataLayerMode.cpp
@@ -0,0 +1,659 @@
+// Copyright Epic Games, Inc. All Rights Reserved.
+
+#include "DataLayerMode.h"
+#include "SSceneOutliner.h"
+#include "SDataLayerBrowser.h"
+#include "DataLayer/DataLayerEditorSubsystem.h"
+#include "DataLayerHierarchy.h"
+#include "DataLayerActorTreeItem.h"
+#include "DataLayerTreeItem.h"
+#include "DataLayerDragDropOp.h"
+#include "ISceneOutlinerHierarchy.h"
+#include "SceneOutlinerMenuContext.h"
+#include "ScopedTransaction.h"
+#include "DragAndDrop/ActorDragDropOp.h"
+#include "DragAndDrop/FolderDragDropOp.h"
+#include "EditorActorFolders.h"
+#include "Algo/Transform.h"
+#include "ToolMenus.h"
+#include "Editor.h"
+
+#define LOCTEXT_NAMESPACE "DataLayer"
+
+FDataLayerModeParams::FDataLayerModeParams(SSceneOutliner* InSceneOutliner, SDataLayerBrowser* InDataLayerBrowser, const TWeakObjectPtr<UWorld>& InSpecifiedWorldToDisplay)
+: SpecifiedWorldToDisplay(InSpecifiedWorldToDisplay)
+, DataLayerBrowser(InDataLayerBrowser)
+, SceneOutliner(InSceneOutliner)
+{}
+
+FDataLayerMode::FDataLayerMode(const FDataLayerModeParams& Params)
+	: ISceneOutlinerMode(Params.SceneOutliner)
+	, DataLayerBrowser(Params.DataLayerBrowser)
+	, SpecifiedWorldToDisplay(Params.SpecifiedWorldToDisplay)
+{
+	DataLayerEditorSubsystem = UDataLayerEditorSubsystem::Get();
+	Rebuild();
+	const_cast<FSharedSceneOutlinerData&>(SceneOutliner->GetSharedData()).CustomDelete = FCustomSceneOutlinerDeleteDelegate::CreateRaw(this, &FDataLayerMode::DeleteItems);
+}
+
+FDataLayerMode::~FDataLayerMode()
+{
+	if (Hierarchy)
+	{
+		Hierarchy->OnHierarchyChanged().RemoveAll(this);
+	}
+}
+
+int32 FDataLayerMode::GetTypeSortPriority(const ISceneOutlinerTreeItem& Item) const
+{
+	if (Item.IsA<FDataLayerTreeItem>())
+	{
+		return (int32)EItemSortOrder::DataLayer;
+	}
+	else if (Item.IsA<FDataLayerActorTreeItem>())
+	{
+		return (int32)EItemSortOrder::Actor;
+	}
+	// Warning: using actor mode with an unsupported item type!
+	check(false);
+	return -1;
+}
+
+SDataLayerBrowser* FDataLayerMode::GetDataLayerBrowser() const
+{
+	return DataLayerBrowser;
+}
+
+void FDataLayerMode::OnItemAdded(FSceneOutlinerTreeItemPtr Item)
+{
+	if (const FDataLayerTreeItem* DataLayerItem = Item->CastTo<FDataLayerTreeItem>())
+	{
+		if (!Item->Flags.bIsFilteredOut)
+		{
+			SceneOutliner->SetItemSelection(Item, true);
+		}
+	}
+}
+
+void FDataLayerMode::OnItemDoubleClick(FSceneOutlinerTreeItemPtr Item)
+{
+	if (const FDataLayerTreeItem* DataLayerItem = Item->CastTo<FDataLayerTreeItem>())
+	{
+		if (const UDataLayer* DataLayer = DataLayerItem->GetDataLayer())
+		{
+			const FScopedTransaction Transaction(LOCTEXT("SelectActorsInDataLayer", "Select Actors in DataLayer"));
+			GEditor->SelectNone(/*bNoteSelectionChange*/false, true);
+			DataLayerEditorSubsystem->SelectActorsInDataLayer(DataLayer, /*bSelect*/true, /*bNotify*/true, /*bSelectEvenIfHidden*/true);
+		}
+	}
+	else if (FDataLayerActorTreeItem* DataLayerActorItem = Item->CastTo<FDataLayerActorTreeItem>())
+	{
+		if (AActor * Actor = DataLayerActorItem->GetActor())
+		{
+			const FScopedTransaction Transaction(LOCTEXT("ClickingOnActor", "Clicking on Actor in DataLayer"));
+			GEditor->GetSelectedActors()->Modify();
+			GEditor->SelectNone(/*bNoteSelectionChange*/false, true);
+			GEditor->SelectActor(Actor, /*bSelected*/true, /*bNotify*/true, /*bSelectEvenIfHidden*/true);
+			GEditor->NoteSelectionChange();
+			GEditor->MoveViewportCamerasToActor(*Actor, /*bActiveViewportOnly*/false);
+		}
+	}
+}
+
+void FDataLayerMode::DeleteItems(const TArray<TWeakPtr<ISceneOutlinerTreeItem>>& Items)
+{
+	TArray<const UDataLayer*> DataLayersToDelete;
+	TMap<const UDataLayer*, TArray<AActor*>> ActorsToRemoveFromDataLayer;
+
+	for (const TWeakPtr<ISceneOutlinerTreeItem>& Item : Items)
+	{
+		if (FDataLayerActorTreeItem* DataLayerActorItem = Item.Pin()->CastTo<FDataLayerActorTreeItem>())
+		{
+			const UDataLayer* DataLayer = DataLayerActorItem->GetDataLayer();
+			AActor* Actor = DataLayerActorItem->GetActor();
+			if (DataLayer && Actor)
+			{
+				ActorsToRemoveFromDataLayer.FindOrAdd(DataLayer).Add(Actor);
+			}
+		}
+		else if (FDataLayerTreeItem* DataLayerItem = Item.Pin()->CastTo< FDataLayerTreeItem>())
+		{
+			if (UDataLayer* DataLayer = DataLayerItem->GetDataLayer())
+			{
+				DataLayersToDelete.Add(DataLayer);
+			}
+		}
+	}
+
+	if (!ActorsToRemoveFromDataLayer.IsEmpty())
+	{
+		const FScopedTransaction Transaction(LOCTEXT("RemoveActorsFromDataLayer", "Remove Actors from Data Layer"));
+		for (const auto& ItPair : ActorsToRemoveFromDataLayer)
+		{
+			DataLayerEditorSubsystem->RemoveActorsFromDataLayer(ItPair.Value, ItPair.Key);
+		}
+	}
+	else if (!DataLayersToDelete.IsEmpty())
+	{
+		const FScopedTransaction Transaction(LOCTEXT("DeleteDataLayers", "Delete DataLayers"));
+		DataLayerEditorSubsystem->DeleteDataLayers(DataLayersToDelete);
+	}
+}
+
+FReply FDataLayerMode::OnKeyDown(const FKeyEvent& InKeyEvent)
+{
+	const FSceneOutlinerItemSelection& Selection = SceneOutliner->GetSelection();
+
+	// Rename key: Rename selected actors (not rebindable, because it doesn't make much sense to bind.)
+	if (InKeyEvent.GetKey() == EKeys::F2)
+	{
+		if (Selection.Num() == 1)
+		{
+			FSceneOutlinerTreeItemPtr ItemToRename = Selection.SelectedItems[0].Pin();
+			if (ItemToRename.IsValid() && CanRenameItem(*ItemToRename) && ItemToRename->CanInteract())
+			{
+				SceneOutliner->SetPendingRenameItem(ItemToRename);
+				SceneOutliner->ScrollItemIntoView(ItemToRename);
+			}
+			return FReply::Handled();
+		}
+	}
+
+	// F5 forces a full refresh
+	else if (InKeyEvent.GetKey() == EKeys::F5)
+	{
+		SceneOutliner->FullRefresh();
+		return FReply::Handled();
+	}
+
+	// Delete/BackSpace keys delete selected actors
+	else if (InKeyEvent.GetKey() == EKeys::Delete || InKeyEvent.GetKey() == EKeys::BackSpace)
+	{
+		DeleteItems(Selection.SelectedItems);
+		return FReply::Handled();
+
+	}
+	return FReply::Unhandled();
+}
+
+bool FDataLayerMode::ParseDragDrop(FSceneOutlinerDragDropPayload& OutPayload, const FDragDropOperation& Operation) const
+{
+	return !GetActorsFromOperation(Operation, true).IsEmpty();
+}
+
+FSceneOutlinerDragValidationInfo FDataLayerMode::ValidateDrop(const ISceneOutlinerTreeItem& DropTarget, const FSceneOutlinerDragDropPayload& Payload) const
+{
+	TArray<AActor*> PayloadActors = GetActorsFromOperation(Payload.SourceOperation, true);
+	if (!PayloadActors.IsEmpty())
+	{
+		if (const FDataLayerTreeItem* DataLayerItem = DropTarget.CastTo<FDataLayerTreeItem>())
+		{
+			const UDataLayer* DataLayerTarget = DataLayerItem->GetDataLayer();
+			if (!DataLayerTarget)
+			{
+				return FSceneOutlinerDragValidationInfo(ESceneOutlinerDropCompatibility::IncompatibleGeneric, FText());
+			}
+			return FSceneOutlinerDragValidationInfo(ESceneOutlinerDropCompatibility::Compatible, FText::Format(LOCTEXT("AssignToDataLayer", "Assign to Data Layer \"{0}\""), FText::FromName(DataLayerTarget->GetDataLayerLabel())));
+		}
+		else
+		{
+			if (!PayloadActors[0]->HasDataLayers())
+			{
+				// Only allow actors not coming from the DataLayerBrowser
+				return FSceneOutlinerDragValidationInfo(ESceneOutlinerDropCompatibility::Compatible, LOCTEXT("AssignToNewDataLayer", "Assign to New Data Layer"));
+			}
+		}
+	}
+	return FSceneOutlinerDragValidationInfo::Invalid();
+}
+
+TArray<AActor*> FDataLayerMode::GetActorsFromOperation(const FDragDropOperation& Operation, bool bOnlyFindFirst) const
+{
+	TSet<AActor*> Actors;
+
+	auto GetActorsFromFolderOperation = [&Actors, bOnlyFindFirst](const FFolderDragDropOp& FolderOp)
+	{
+		if (bOnlyFindFirst && Actors.Num())
+		{
+			return;
+		}
+
+		if (UWorld* World = FolderOp.World.Get())
+		{
+			TArray<TWeakObjectPtr<AActor>> ActorsToDrop;
+			FActorFolders::GetWeakActorsFromFolders(*World, FolderOp.Folders, ActorsToDrop);
+			for (const auto& Actor : ActorsToDrop)
+			{
+				if (AActor* ActorPtr = Actor.Get())
+				{
+					Actors.Add(ActorPtr);
+					if (bOnlyFindFirst)
+					{
+						break;
+					}
+				}
+			}
+		}
+	};
+
+	auto GetActorsFromActorOperation = [&Actors, bOnlyFindFirst](const FActorDragDropOp& ActorOp)
+	{
+		if (bOnlyFindFirst && Actors.Num())
+		{
+			return;
+		}
+		for (const auto& Actor : ActorOp.Actors)
+		{
+			if (AActor* ActorPtr = Actor.Get())
+			{
+				Actors.Add(ActorPtr);
+				if (bOnlyFindFirst)
+				{
+					break;
+				}
+			}
+		}
+	};
+
+	TSharedPtr<FActorDragDropOp> ActorDragOp = const_cast<FDragDropOperation&>(Operation).CastTo<FActorDragDropOp>();
+	if (ActorDragOp.IsValid())
+	{
+		GetActorsFromActorOperation(*ActorDragOp);
+	}
+	TSharedPtr<FFolderDragDropOp> FolderDragOp = const_cast<FDragDropOperation&>(Operation).CastTo<FFolderDragDropOp>();
+	if (FolderDragOp.IsValid())
+	{
+		GetActorsFromFolderOperation(*FolderDragOp);
+	}
+	return Actors.Array();
+}
+
+void FDataLayerMode::OnDrop(ISceneOutlinerTreeItem& DropTarget, const FSceneOutlinerDragDropPayload& Payload, const FSceneOutlinerDragValidationInfo& ValidationInfo) const
+{
+	TArray<AActor*> ActorsToAddToDataLayer = GetActorsFromOperation(Payload.SourceOperation);
+	if (ActorsToAddToDataLayer.IsEmpty())
+	{
+		return;
+	}
+
+	if (const FDataLayerTreeItem* DataLayerItem = DropTarget.CastTo<FDataLayerTreeItem>())
+	{
+		if (const UDataLayer* DataLayer = DataLayerItem->GetDataLayer())
+		{
+			const FScopedTransaction Transaction(LOCTEXT("DataLayerOutlinerAddActorsToDataLayer", "Add Actors to DataLayers"));
+			DataLayerEditorSubsystem->AddActorsToDataLayer(ActorsToAddToDataLayer, DataLayer);
+		}
+	}
+	else
+	{
+		if (!ActorsToAddToDataLayer[0]->HasDataLayers())
+		{
+			// Only allow actors not coming from the DataLayerBrowser
+			const FScopedTransaction Transaction(LOCTEXT("AddSelectedActorsToNewDataLayer", "Add Actors to New DataLayer"));
+			if (const UDataLayer* NewDataLayer = DataLayerEditorSubsystem->CreateDataLayer())
+			{
+				DataLayerEditorSubsystem->AddActorsToDataLayer(ActorsToAddToDataLayer, NewDataLayer);
+			}
+		}
+	}
+}
+
+struct FWeakDataLayerActorSelector
+{
+	bool operator()(const TWeakPtr<ISceneOutlinerTreeItem>& Item, TWeakObjectPtr<AActor>& DataOut) const
+	{
+		if (TSharedPtr<ISceneOutlinerTreeItem> ItemPtr = Item.Pin())
+		{
+			if (FDataLayerActorTreeItem* TypedItem = ItemPtr->CastTo<FDataLayerActorTreeItem>())
+			{
+				if (TypedItem->IsValid())
+				{
+					DataOut = TypedItem->Actor;
+					return true;
+				}
+			}
+		}
+		return false;
+	}
+};
+
+struct FWeakDataLayerSelector
+{
+	bool operator()(const TWeakPtr<ISceneOutlinerTreeItem>& Item, TWeakObjectPtr<UDataLayer>& DataOut) const
+	{
+		if (TSharedPtr<ISceneOutlinerTreeItem> ItemPtr = Item.Pin())
+		{
+			if (FDataLayerTreeItem* TypedItem = ItemPtr->CastTo<FDataLayerTreeItem>())
+			{
+				if (TypedItem->IsValid())
+				{
+					DataOut = TypedItem->GetDataLayer();
+					return true;
+				}
+			}
+		}
+		return false;
+	}
+};
+
+TSharedPtr<FDragDropOperation> FDataLayerMode::CreateDragDropOperation(const TArray<FSceneOutlinerTreeItemPtr>& InTreeItems) const
+{
+	FSceneOutlinerDragDropPayload DraggedObjects(InTreeItems);
+
+	TSharedPtr<FSceneOutlinerDragDropOp> OutlinerOp = MakeShareable(new FSceneOutlinerDragDropOp());
+
+	if (DraggedObjects.Has<FDataLayerActorTreeItem>())
+	{
+		TSharedPtr<FActorDragDropOp> ActorOperation = MakeShareable(new FActorDragDropOp);
+		ActorOperation->Init(DraggedObjects.GetData<TWeakObjectPtr<AActor>>(FWeakDataLayerActorSelector()));
+		OutlinerOp->AddSubOp(ActorOperation);
+	}
+
+	if (DraggedObjects.Has<FDataLayerTreeItem>())
+	{
+		TSharedPtr<FDataLayerDragDropOp> DataLayerOperation = MakeShareable(new FDataLayerDragDropOp);
+		TArray<TWeakObjectPtr<UDataLayer>> DataLayers = DraggedObjects.GetData<TWeakObjectPtr<UDataLayer>>(FWeakDataLayerSelector());
+		for (const auto& DataLayer : DataLayers)
+		{
+			if (DataLayer.IsValid())
+			{
+				DataLayerOperation->DataLayerLabels.Add(DataLayer->GetDataLayerLabel());
+			}
+		}
+		DataLayerOperation->Construct();
+		OutlinerOp->AddSubOp(DataLayerOperation);
+	}
+
+	OutlinerOp->Construct();
+	return OutlinerOp;
+}
+
+static const FName DefaultContextBaseMenuName("DataLayerOutliner.DefaultContextMenuBase");
+static const FName DefaultContextMenuName("DataLayerOutliner.DefaultContextMenu");
+
+void FDataLayerMode::RegisterContextMenu()
+{
+	UToolMenus* ToolMenus = UToolMenus::Get();
+
+	if (!ToolMenus->IsMenuRegistered(DefaultContextBaseMenuName))
+	{
+		UToolMenu* Menu = ToolMenus->RegisterMenu(DefaultContextBaseMenuName);
+
+		Menu->AddDynamicSection("DataLayerDynamicSection", FNewToolMenuDelegate::CreateLambda([this](UToolMenu* InMenu)
+		{
+			USceneOutlinerMenuContext* Context = InMenu->FindContext<USceneOutlinerMenuContext>();
+			if (!Context || !Context->SceneOutliner.IsValid())
+			{
+				return;
+			}
+
+			SSceneOutliner* SceneOutliner = Context->SceneOutliner.Pin().Get();
+			FSceneOutlinerItemSelection ItemSelection(SceneOutliner->GetSelection());
+			TArray<FDataLayerTreeItem*> SelectedDataLayerItems;
+			ItemSelection.Get<FDataLayerTreeItem>(SelectedDataLayerItems);
+			TArray<const UDataLayer*> SelectedDataLayers;
+			Algo::TransformIf(SelectedDataLayerItems, SelectedDataLayers, [](const auto Item) { return Item && Item->GetDataLayer(); }, [](const auto Item) { return Item->GetDataLayer(); });
+
+			TArray<const UDataLayer*> AllDataLayers;
+			if (const AWorldDataLayers* WorldDataLayers = AWorldDataLayers::Get(RepresentingWorld.Get()))
+			{
+				WorldDataLayers->ForEachDataLayer([&AllDataLayers](UDataLayer* DataLayer)
+					{
+						AllDataLayers.Add(DataLayer);
+						return true;
+					});
+			}
+
+			{
+				FToolMenuSection& Section = InMenu->AddSection("DataLayers", LOCTEXT("DataLayers", "DataLayers"));
+				Section.AddMenuEntry("CreateEmptyDataLayer", LOCTEXT("CreateEmptyDataLayer", "Create Empty DataLayer"), FText(), FSlateIcon(),
+					FUIAction(
+						FExecuteAction::CreateLambda([=]() {
+							const FScopedTransaction Transaction(LOCTEXT("CreateEmptyDataLayer", "Create Empty DataLayer"));
+							DataLayerEditorSubsystem->CreateDataLayer();
+							})
+					));
+
+				Section.AddMenuEntry("AddSelectedActorsToNewDataLayer", LOCTEXT("AddSelectedActorsToNewDataLayer", "Add Selected Actors to New DataLayer"), FText(), FSlateIcon(),
+					FUIAction(
+						FExecuteAction::CreateLambda([=]() {
+							const FScopedTransaction Transaction(LOCTEXT("AddSelectedActorsToNewDataLayer", "Add Selected Actors to New DataLayer"));
+							if (UDataLayer* NewDataLayer = DataLayerEditorSubsystem->CreateDataLayer())
+							{
+								DataLayerEditorSubsystem->AddSelectedActorsToDataLayer(NewDataLayer);
+							}}),
+						FCanExecuteAction::CreateLambda([=] { return GEditor->GetSelectedActorCount() > 0; })
+								));
+
+				Section.AddMenuEntry("AddSelectedActorsToSelectedDataLayers", LOCTEXT("AddSelectedActorsToSelectedDataLayers", "Add Selected Actors to Selected DataLayers"), FText(), FSlateIcon(),
+					FUIAction(
+						FExecuteAction::CreateLambda([=]() {
+							check(!SelectedDataLayers.IsEmpty());
+							{
+								const FScopedTransaction Transaction(LOCTEXT("AddSelectedActorsToSelectedDataLayers", "Add Selected Actors to Selected DataLayers"));
+								DataLayerEditorSubsystem->AddSelectedActorsToDataLayers(SelectedDataLayers);
+							}}),
+						FCanExecuteAction::CreateLambda([=] { return !SelectedDataLayers.IsEmpty() && GEditor->GetSelectedActorCount() > 0; })
+								));
+
+				Section.AddSeparator("SectionsSeparator");
+
+				Section.AddMenuEntry("RemoveSelectedActorsFromSelectedDataLayers", LOCTEXT("RemoveSelectedActorsFromSelectedDataLayers", "Remove Selected Actors from Selected DataLayers"), FText(), FSlateIcon(),
+					FUIAction(
+						FExecuteAction::CreateLambda([=]() {
+							check(!SelectedDataLayers.IsEmpty());
+							{
+								const FScopedTransaction Transaction(LOCTEXT("RemoveSelectedActorsFromSelectedDataLayers", "Remove Selected Actors from Selected DataLayers"));
+								DataLayerEditorSubsystem->RemoveSelectedActorsFromDataLayers(SelectedDataLayers);
+							}}),
+						FCanExecuteAction::CreateLambda([=] { return !SelectedDataLayers.IsEmpty() && GEditor->GetSelectedActorCount() > 0; })
+								));
+
+				Section.AddMenuEntry("DeleteSelectedDataLayers", LOCTEXT("DeleteSelectedDataLayers", "Delete Selected DataLayers"), FText(), FSlateIcon(),
+					FUIAction(
+						FExecuteAction::CreateLambda([=]() {
+							check(!SelectedDataLayers.IsEmpty());
+							{
+								const FScopedTransaction Transaction(LOCTEXT("DeleteSelectedDataLayers", "Delete Selected DataLayers"));
+								DataLayerEditorSubsystem->DeleteDataLayers(SelectedDataLayers);
+							}}),
+						FCanExecuteAction::CreateLambda([=] { return !SelectedDataLayers.IsEmpty(); })
+								));
+
+				Section.AddMenuEntry("RenameSelectedDataLayer", LOCTEXT("RenameSelectedDataLayer", "Rename Selected DataLayer"), FText(), FSlateIcon(),
+					FUIAction(
+						FExecuteAction::CreateLambda([=]() {
+							if (SelectedDataLayers.Num() == 1)
+							{
+								FSceneOutlinerTreeItemPtr ItemToRename = SceneOutliner->GetTreeItem(SelectedDataLayers[0]);
+								if (ItemToRename.IsValid() && CanRenameItem(*ItemToRename) && ItemToRename->CanInteract())
+								{
+									SceneOutliner->SetPendingRenameItem(ItemToRename);
+									SceneOutliner->ScrollItemIntoView(ItemToRename);
+								}
+							}}),
+						FCanExecuteAction::CreateLambda([=] { return SelectedDataLayers.Num() == 1; })
+								));
+
+				Section.AddSeparator("SectionsSeparator");
+			}
+
+			{
+				FToolMenuSection& Section = InMenu->AddSection("DataLayerSelection", LOCTEXT("DataLayerSelection", "Selection"));
+
+				Section.AddMenuEntry("SelectActorsInDataLayers", LOCTEXT("SelectActorsInDataLayers", "Select Actors in DataLayers"), FText(), FSlateIcon(),
+					FUIAction(
+						FExecuteAction::CreateLambda([=]() {
+							check(!SelectedDataLayers.IsEmpty());
+							{
+								const FScopedTransaction Transaction(LOCTEXT("SelectActorsInDataLayers", "Select Actors in DataLayers"));
+								GEditor->SelectNone(/*bNoteSelectionChange*/false, /*bDeselectBSPSurfs*/true);
+								DataLayerEditorSubsystem->SelectActorsInDataLayers(SelectedDataLayers, /*bSelect*/true, /*bNotify*/true, /*bSelectEvenIfHidden*/true);
+							}}),
+						FCanExecuteAction::CreateLambda([=] { return !SelectedDataLayers.IsEmpty(); })
+								));
+
+				Section.AddMenuEntry("AppendActorsToSelection", LOCTEXT("AppendActorsToSelection", "Append Actors in DataLayer to Selection"), FText(), FSlateIcon(),
+					FUIAction(
+						FExecuteAction::CreateLambda([=]() {
+							check(!SelectedDataLayers.IsEmpty());
+							{
+								const FScopedTransaction Transaction(LOCTEXT("AppendActorsToSelection", "Append Actors in DataLayer to Selection"));
+								DataLayerEditorSubsystem->SelectActorsInDataLayers(SelectedDataLayers, /*bSelect*/true, /*bNotify*/true, /*bSelectEvenIfHidden*/true);
+							}}),
+						FCanExecuteAction::CreateLambda([=] { return !SelectedDataLayers.IsEmpty(); })
+								));
+
+				Section.AddMenuEntry("DeselectActors", LOCTEXT("DeselectActors", "Deselect Actors in DataLayer"), FText(), FSlateIcon(),
+					FUIAction(
+						FExecuteAction::CreateLambda([=]() {
+							check(!SelectedDataLayers.IsEmpty());
+							{
+								const FScopedTransaction Transaction(LOCTEXT("DeselectActors", "Deselect Actors in DataLayer"));
+								DataLayerEditorSubsystem->SelectActorsInDataLayers(SelectedDataLayers, /*bSelect*/false, /*bNotifySelectActors*/true);
+							}}),
+						FCanExecuteAction::CreateLambda([=] { return !SelectedDataLayers.IsEmpty(); })
+								));
+			}
+
+			{
+				FToolMenuSection& Section = InMenu->AddSection("DataLayerVisibility", LOCTEXT("DataLayerVisibility", "Visibility"));
+
+				Section.AddMenuEntry("MakeAllDataLayersVisible", LOCTEXT("MakeAllDataLayersVisible", "Make All DataLayers Visible"), FText(), FSlateIcon(),
+					FUIAction(
+						FExecuteAction::CreateLambda([=]() {
+							check(!AllDataLayers.IsEmpty());
+							{
+								const FScopedTransaction Transaction(LOCTEXT("MakeAllDataLayersVisible", "Make All DataLayers Visible"));
+								DataLayerEditorSubsystem->MakeAllDataLayersVisible();
+							}}),
+						FCanExecuteAction::CreateLambda([=] { return !AllDataLayers.IsEmpty(); })
+								));
+
+
+			}
+		}));
+	}
+
+	if (!ToolMenus->IsMenuRegistered(DefaultContextMenuName))
+	{
+		ToolMenus->RegisterMenu(DefaultContextMenuName, DefaultContextBaseMenuName);
+	}
+}
+
+TSharedPtr<SWidget> FDataLayerMode::CreateContextMenu()
+{
+	RegisterContextMenu();
+
+	FSceneOutlinerItemSelection ItemSelection(SceneOutliner->GetSelection());
+
+	USceneOutlinerMenuContext* ContextObject = NewObject<USceneOutlinerMenuContext>();
+	ContextObject->SceneOutliner = StaticCastSharedRef<SSceneOutliner>(SceneOutliner->AsShared());
+	ContextObject->bShowParentTree = SceneOutliner->GetSharedData().bShowParentTree;
+	ContextObject->NumSelectedItems = ItemSelection.Num();
+	FToolMenuContext Context(ContextObject);
+
+	FName MenuName = DefaultContextMenuName;
+	SceneOutliner->GetSharedData().ModifyContextMenu.ExecuteIfBound(MenuName, Context);
+
+	// Build up the menu for a selection
+	UToolMenus* ToolMenus = UToolMenus::Get();
+	UToolMenu* Menu = ToolMenus->GenerateMenu(MenuName, Context);
+	for (const FToolMenuSection& Section : Menu->Sections)
+	{
+		if (Section.Blocks.Num() > 0)
+		{
+			return ToolMenus->GenerateWidget(Menu);
+		}
+	}
+
+	return nullptr;
+}
+
+TUniquePtr<ISceneOutlinerHierarchy> FDataLayerMode::CreateHierarchy()
+{
+	TUniquePtr<FDataLayerHierarchy> ActorHierarchy = FDataLayerHierarchy::Create(this, RepresentingWorld);
+	return ActorHierarchy;
+}
+
+void FDataLayerMode::Rebuild()
+{
+	ChooseRepresentingWorld();
+
+	if (Hierarchy)
+	{
+		Hierarchy->OnHierarchyChanged().Clear();
+	}
+
+	Hierarchy = CreateHierarchy();
+}
+
+void FDataLayerMode::ChooseRepresentingWorld()
+{
+	// Select a world to represent
+
+	RepresentingWorld = nullptr;
+
+	// If a specified world was provided, represent it
+	if (SpecifiedWorldToDisplay.IsValid())
+	{
+		RepresentingWorld = SpecifiedWorldToDisplay.Get();
+	}
+
+	// check if the user-chosen world is valid and in the editor contexts
+
+	if (!RepresentingWorld.IsValid() && UserChosenWorld.IsValid())
+	{
+		for (const FWorldContext& Context : GEngine->GetWorldContexts())
+		{
+			if (UserChosenWorld.Get() == Context.World())
+			{
+				RepresentingWorld = UserChosenWorld.Get();
+				break;
+			}
+		}
+	}
+
+	// If the user did not manually select a world, try to pick the most suitable world context
+	if (!RepresentingWorld.IsValid())
+	{
+		// ideally we want a PIE world that is standalone or the first client
+		for (const FWorldContext& Context : GEngine->GetWorldContexts())
+		{
+			UWorld* World = Context.World();
+			if (World && Context.WorldType == EWorldType::PIE)
+			{
+				if (World->GetNetMode() == NM_Standalone)
+				{
+					RepresentingWorld = World;
+					break;
+				}
+				else if (World->GetNetMode() == NM_Client && Context.PIEInstance == 2)	// Slightly dangerous: assumes server is always PIEInstance = 1;
+				{
+					RepresentingWorld = World;
+					break;
+				}
+			}
+		}
+	}
+
+	if (RepresentingWorld == nullptr)
+	{
+		// still not world so fallback to old logic where we just prefer PIE over Editor
+
+		for (const FWorldContext& Context : GEngine->GetWorldContexts())
+		{
+			if (Context.WorldType == EWorldType::PIE)
+			{
+				RepresentingWorld = Context.World();
+				break;
+			}
+			else if (Context.WorldType == EWorldType::Editor)
+			{
+				RepresentingWorld = Context.World();
+			}
+		}
+	}
+}
+
+#undef LOCTEXT_NAMESPACE
\ No newline at end of file
diff --git a/Engine/Source/Editor/DataLayerEditor/Private/DataLayer/DataLayerMode.h b/Engine/Source/Editor/DataLayerEditor/Private/DataLayer/DataLayerMode.h
new file mode 100644
--- /dev/null
+++ b/Engine/Source/Editor/DataLayerEditor/Private/DataLayer/DataLayerMode.h
@@ -0,0 +1,80 @@
+// Copyright Epic Games, Inc. All Rights Reserved.
+
+#pragma once
+
+#include "CoreMinimal.h"
+#include "SceneOutlinerFwd.h"
+#include "ISceneOutlinerMode.h"
+#include "DataLayerTreeItem.h"
+
+class UDataLayerEditorSubsystem;
+class SDataLayerBrowser;
+class UWorld;
+
+struct FDataLayerModeParams
+{
+	FDataLayerModeParams() {}
+
+	FDataLayerModeParams(SSceneOutliner* InSceneOutliner, SDataLayerBrowser* InDataLayerBrowser, const TWeakObjectPtr<UWorld>& InSpecifiedWorldToDisplay = nullptr);
+
+	TWeakObjectPtr<UWorld> SpecifiedWorldToDisplay = nullptr;
+	SDataLayerBrowser* DataLayerBrowser;
+	SSceneOutliner* SceneOutliner;
+};
+
+class FDataLayerMode : public ISceneOutlinerMode
+{
+public:
+	enum class EItemSortOrder : int32
+	{
+		DataLayer = 0,
+		Actor = 10
+	};
+
+	FDataLayerMode(const FDataLayerModeParams& Params);
+	virtual ~FDataLayerMode();
+
+	virtual void Rebuild() override;
+	virtual TSharedPtr<SWidget> CreateContextMenu() override;
+
+	virtual int32 GetTypeSortPriority(const ISceneOutlinerTreeItem& Item) const override;
+
+	virtual ESelectionMode::Type GetSelectionMode() const override { return ESelectionMode::Multi; }
+	virtual bool SupportsKeyboardFocus() const override { return true; }
+	virtual bool IsInteractive() const override { return true; }
+	virtual bool CanRename() const override { return true; }
+	virtual bool CanRenameItem(const ISceneOutlinerTreeItem& Item) const override { return (Item.IsValid() && (Item.IsA<FDataLayerTreeItem>())); }
+
+	virtual void OnItemDoubleClick(FSceneOutlinerTreeItemPtr Item) override;
+	virtual void OnItemAdded(FSceneOutlinerTreeItemPtr Item) override;
+	virtual FReply OnKeyDown(const FKeyEvent& InKeyEvent) override;
+
+	virtual bool CanSupportDragAndDrop() const { return true; }
+	virtual bool ParseDragDrop(FSceneOutlinerDragDropPayload& OutPayload, const FDragDropOperation& Operation) const override;
+	virtual FSceneOutlinerDragValidationInfo ValidateDrop(const ISceneOutlinerTreeItem& DropTarget, const FSceneOutlinerDragDropPayload& Payload) const override;
+	virtual void OnDrop(ISceneOutlinerTreeItem& DropTarget, const FSceneOutlinerDragDropPayload& Payload, const FSceneOutlinerDragValidationInfo& ValidationInfo) const override;
+	virtual TSharedPtr<FDragDropOperation> CreateDragDropOperation(const TArray<FSceneOutlinerTreeItemPtr>& InTreeItems) const override;
+
+	void DeleteItems(const TArray<TWeakPtr<ISceneOutlinerTreeItem>>& Items);
+	SDataLayerBrowser* GetDataLayerBrowser() const;
+
+protected:
+	virtual TUniquePtr<ISceneOutlinerHierarchy> CreateHierarchy() override;
+
+private:
+	/* Private Helpers */
+	void RegisterContextMenu();
+	void ChooseRepresentingWorld();
+	TArray<AActor*> GetActorsFromOperation(const FDragDropOperation& Operation, bool bOnlyFindFirst = false) const;
+
+	/** The world which we are currently representing */
+	TWeakObjectPtr<UWorld> RepresentingWorld;
+	/** The world which the user manually selected */
+	TWeakObjectPtr<UWorld> UserChosenWorld;
+	/** The DataLayer browser */
+	SDataLayerBrowser* DataLayerBrowser;
+	/** The DataLayerEditorSubsystem */
+	UDataLayerEditorSubsystem* DataLayerEditorSubsystem;
+	/** If this mode was created to display a specific world, don't allow it to be reassigned */
+	const TWeakObjectPtr<UWorld> SpecifiedWorldToDisplay;
+};
diff --git a/Engine/Source/Editor/DataLayerEditor/Private/DataLayer/DataLayerModeSettings.h b/Engine/Source/Editor/DataLayerEditor/Private/DataLayer/DataLayerModeSettings.h
new file mode 100644
--- /dev/null
+++ b/Engine/Source/Editor/DataLayerEditor/Private/DataLayer/DataLayerModeSettings.h
@@ -0,0 +1,18 @@
+// Copyright Epic Games, Inc. All Rights Reserved.
+
+#pragma once
+
+#include "CoreMinimal.h"
+#include "UObject/ObjectMacros.h"
+#include "UObject/Object.h"
+#include "DataLayerModeSettings.generated.h"
+
+UCLASS(config = EditorPerProjectUserSettings)
+class UDataLayerModeSettings : public UObject
+{
+	GENERATED_BODY()
+
+public:
+	UPROPERTY(config)
+	uint32 bShowDataLayerContent : 1;
+};
diff --git a/Engine/Source/Editor/DataLayerEditor/Private/DataLayer/DataLayerOutlinerDeleteButtonColumn.cpp b/Engine/Source/Editor/DataLayerEditor/Private/DataLayer/DataLayerOutlinerDeleteButtonColumn.cpp
new file mode 100644
--- /dev/null
+++ b/Engine/Source/Editor/DataLayerEditor/Private/DataLayer/DataLayerOutlinerDeleteButtonColumn.cpp
@@ -0,0 +1,83 @@
+// Copyright Epic Games, Inc. All Rights Reserved.
+
+#include "DataLayerOutlinerDeleteButtonColumn.h"
+#include "Widgets/Layout/SSpacer.h"
+#include "Widgets/Input/SButton.h"
+#include "Widgets/Views/SHeaderRow.h"
+#include "Widgets/Views/STreeView.h"
+#include "Widgets/Images/SImage.h"
+#include "DataLayerActorTreeItem.h"
+#include "DataLayer/DataLayerEditorSubsystem.h"
+#include "SceneOutlinerPublicTypes.h"
+#include "ScopedTransaction.h"
+#include "EditorStyleSet.h"
+#include "Editor.h"
+
+#define LOCTEXT_NAMESPACE "DataLayer"
+
+FName FDataLayerOutlinerDeleteButtonColumn::GetID()
+{
+	static FName DataLayeDeleteButton("DataLayeDeleteButton");
+	return DataLayeDeleteButton;
+}
+
+SHeaderRow::FColumn::FArguments FDataLayerOutlinerDeleteButtonColumn::ConstructHeaderRowColumn()
+{
+	return SHeaderRow::Column(GetColumnID())
+	[
+		SNew(SSpacer)
+	];
+}
+
+const TSharedRef<SWidget> FDataLayerOutlinerDeleteButtonColumn::ConstructRowWidget(FSceneOutlinerTreeItemRef TreeItem, const STableRow<FSceneOutlinerTreeItemPtr>& Row)
+{
+	if (FDataLayerActorTreeItem* DataLayerActorItem = TreeItem->CastTo<FDataLayerActorTreeItem>())
+	{
+		return SNew(SButton)
+			.HAlign(HAlign_Center)
+			.VAlign(VAlign_Center)
+			.ButtonStyle(FEditorStyle::Get(), "DataLayerBrowserButton")
+			.ContentPadding(0)
+			.OnClicked_Lambda([this, TreeItem, DataLayerActorItem]()
+			{
+				AActor* Actor = DataLayerActorItem->GetActor();
+				const UDataLayer* DataLayer = DataLayerActorItem->GetDataLayer();
+				if (Actor && DataLayer)
+				{
+					UDataLayerEditorSubsystem* DataLayerEditorSubsystem = UDataLayerEditorSubsystem::Get();
+					if (auto SceneOutliner = WeakSceneOutliner.IsValid() ? WeakSceneOutliner.Pin() : nullptr)
+					{
+						const auto& Tree = SceneOutliner->GetTree();
+						if (SceneOutliner->GetSharedData().CustomDelete.IsBound())
+						{
+							TArray<TWeakPtr<ISceneOutlinerTreeItem>> SelectedItems;
+							if (Tree.IsItemSelected(TreeItem))
+							{
+								for (auto& SelectedItem : Tree.GetSelectedItems())
+								{
+									if (FDataLayerActorTreeItem* SelectedDataLayerActorTreeItem = SelectedItem->CastTo<FDataLayerActorTreeItem>())
+									{
+										SelectedItems.Add(SelectedDataLayerActorTreeItem->AsShared());
+									}
+								}
+							}
+							else
+							{
+								SelectedItems.Add(TreeItem);
+							}
+							SceneOutliner->GetSharedData().CustomDelete.Execute(SelectedItems);
+						}
+					}
+				}
+				return FReply::Handled();
+			})
+			.ToolTipText(LOCTEXT("RemoveFromDataLayerButtonText", "Remove from DataLayer"))
+			[
+				SNew(SImage)
+				.Image(FEditorStyle::GetBrush(TEXT("DataLayerBrowser.Actor.RemoveFromDataLayer")))
+			];
+	}
+	return SNullWidget::NullWidget;
+}
+
+#undef LOCTEXT_NAMESPACE
\ No newline at end of file
diff --git a/Engine/Source/Editor/DataLayerEditor/Private/DataLayer/DataLayerOutlinerDeleteButtonColumn.h b/Engine/Source/Editor/DataLayerEditor/Private/DataLayer/DataLayerOutlinerDeleteButtonColumn.h
new file mode 100644
--- /dev/null
+++ b/Engine/Source/Editor/DataLayerEditor/Private/DataLayer/DataLayerOutlinerDeleteButtonColumn.h
@@ -0,0 +1,31 @@
+// Copyright Epic Games, Inc. All Rights Reserved.
+
+#pragma once
+
+#include "CoreMinimal.h"
+#include "SceneOutlinerFwd.h"
+#include "ISceneOutliner.h"
+#include "ISceneOutlinerColumn.h"
+#include "Widgets/Views/SHeaderRow.h"
+
+class FDataLayerOutlinerDeleteButtonColumn : public ISceneOutlinerColumn
+{
+public:
+
+	FDataLayerOutlinerDeleteButtonColumn(ISceneOutliner& SceneOutliner) : WeakSceneOutliner(StaticCastSharedRef<ISceneOutliner>(SceneOutliner.AsShared())) {}
+	virtual ~FDataLayerOutlinerDeleteButtonColumn() {}
+	static FName GetID();
+
+	//////////////////////////////////////////////////////////////////////////
+	// Begin ISceneOutlinerColumn Implementation
+	virtual FName GetColumnID() override { return GetID(); }
+	virtual SHeaderRow::FColumn::FArguments ConstructHeaderRowColumn() override;
+	virtual const TSharedRef<SWidget> ConstructRowWidget(FSceneOutlinerTreeItemRef TreeItem, const STableRow<FSceneOutlinerTreeItemPtr>& Row) override;
+	virtual bool SupportsSorting() const override { return false; }
+	// End ISceneOutlinerColumn Implementation
+	//////////////////////////////////////////////////////////////////////////
+
+private:
+
+	TWeakPtr<ISceneOutliner> WeakSceneOutliner;
+};
\ No newline at end of file
diff --git a/Engine/Source/Editor/DataLayerEditor/Private/DataLayer/DataLayerOutlinerIsDynamicallyLoadedColumn.cpp b/Engine/Source/Editor/DataLayerEditor/Private/DataLayer/DataLayerOutlinerIsDynamicallyLoadedColumn.cpp
new file mode 100644
--- /dev/null
+++ b/Engine/Source/Editor/DataLayerEditor/Private/DataLayer/DataLayerOutlinerIsDynamicallyLoadedColumn.cpp
@@ -0,0 +1,106 @@
+// Copyright Epic Games, Inc. All Rights Reserved.
+
+#include "DataLayerOutlinerIsDynamicallyLoadedColumn.h"
+#include "DataLayer/DataLayerEditorSubsystem.h"
+#include "DataLayerTreeItem.h"
+#include "Widgets/Images/SImage.h"
+#include "Widgets/Input/SButton.h"
+#include "Widgets/Views/STreeView.h"
+#include "ScopedTransaction.h"
+#include "EditorStyleSet.h"
+#include "Editor.h"
+
+#define LOCTEXT_NAMESPACE "DataLayer"
+
+FName FDataLayerOutlinerIsDynamicallyLoadedColumn::GetID()
+{
+	static FName DataLayerIsDynamicallyLoaded("DataLayerIsDynamicallyLoaded");
+	return DataLayerIsDynamicallyLoaded;
+}
+
+SHeaderRow::FColumn::FArguments FDataLayerOutlinerIsDynamicallyLoadedColumn::ConstructHeaderRowColumn()
+{
+	return SHeaderRow::Column(GetColumnID())
+		.FixedWidth(24.f)
+		.HAlignHeader(HAlign_Center)
+		.VAlignHeader(VAlign_Center)
+		.HAlignCell(HAlign_Center)
+		.VAlignCell(VAlign_Center)
+		[
+			SNew(SImage)
+			.ColorAndOpacity(FSlateColor::UseForeground())
+		];
+}
+
+const TSharedRef<SWidget> FDataLayerOutlinerIsDynamicallyLoadedColumn::ConstructRowWidget(FSceneOutlinerTreeItemRef TreeItem, const STableRow<FSceneOutlinerTreeItemPtr>& Row)
+{
+	if (TreeItem->IsA<FDataLayerTreeItem>())
+	{
+		return SNew(SHorizontalBox)
+			+ SHorizontalBox::Slot()
+			.AutoWidth()
+			.VAlign(VAlign_Center)
+			[
+				SNew(SButton)
+				.ContentPadding(0)
+				.ButtonStyle(FEditorStyle::Get(), "NoBorder")
+				.OnClicked_Lambda([this, TreeItem]()
+				{
+					FDataLayerTreeItem* DataLayerTreeItem = TreeItem->CastTo<FDataLayerTreeItem>();
+					if (UDataLayer* DataLayer = DataLayerTreeItem->GetDataLayer())
+					{
+						UDataLayerEditorSubsystem* DataLayerEditorSubsystem = UDataLayerEditorSubsystem::Get();
+						const auto& Tree = WeakSceneOutliner.Pin()->GetTree();
+						if (Tree.IsItemSelected(TreeItem))
+						{
+							// Toggle IsDynamicallyLoaded flag of selected DataLayers to the same state as the given DataLayer
+							bool bIsDynamicallyLoaded = DataLayer->IsDynamicallyLoaded();
+
+							TArray<UDataLayer*> AllSelectedDataLayers;
+							for (auto& SelectedItem : Tree.GetSelectedItems())
+							{
+								FDataLayerTreeItem* SelectedDataLayerTreeItem = SelectedItem->CastTo<FDataLayerTreeItem>();
+								UDataLayer* SelectedDataLayer = SelectedDataLayerTreeItem ? SelectedDataLayerTreeItem->GetDataLayer() : nullptr;
+								if (SelectedDataLayer && SelectedDataLayer->IsDynamicallyLoaded() == bIsDynamicallyLoaded)
+								{
+									AllSelectedDataLayers.Add(SelectedDataLayer);
+								}
+							}
+
+							const FScopedTransaction Transaction(LOCTEXT("ToggleDataLayersIsDynamicallyLoaded", "Toggle DataLayers Actors Loading"));
+							DataLayerEditorSubsystem->ToggleDataLayersIsDynamicallyLoaded(AllSelectedDataLayers);
+						}
+						else
+						{
+							const FScopedTransaction Transaction(LOCTEXT("ToggleDataLayerIsDynamicallyLoaded", "Toggle DataLayer Actors Loading"));
+							DataLayerEditorSubsystem->ToggleDataLayerIsDynamicallyLoaded(DataLayer);
+						}
+					}
+					return FReply::Handled();
+				})
+				.ToolTipText(LOCTEXT("IsDynamicallyLoadedButtonToolTip", "Toggle DataLayer Actors Loading"))
+				.HAlign(HAlign_Center)
+				.VAlign(VAlign_Center)
+				.Content()
+				[
+					SNew(SImage)
+					.Image_Lambda([this, TreeItem]()
+					{
+						FDataLayerTreeItem* DataLayerTreeItem = TreeItem->CastTo<FDataLayerTreeItem>();
+						const UDataLayer* DataLayer = DataLayerTreeItem->GetDataLayer();
+						if (DataLayer && DataLayer->IsDynamicallyLoaded())
+						{
+							return FEditorStyle::GetBrush("DataLayer.DynamicallyLoaded");
+						}
+						else
+						{
+							return FEditorStyle::GetBrush("DataLayer.NotDynamicallyLoaded");
+						}
+					})
+				]
+			];
+	}
+	return SNullWidget::NullWidget;
+}
+
+#undef LOCTEXT_NAMESPACE
\ No newline at end of file
diff --git a/Engine/Source/Editor/DataLayerEditor/Private/DataLayer/DataLayerOutlinerIsDynamicallyLoadedColumn.h b/Engine/Source/Editor/DataLayerEditor/Private/DataLayer/DataLayerOutlinerIsDynamicallyLoadedColumn.h
new file mode 100644
--- /dev/null
+++ b/Engine/Source/Editor/DataLayerEditor/Private/DataLayer/DataLayerOutlinerIsDynamicallyLoadedColumn.h
@@ -0,0 +1,29 @@
+// Copyright Epic Games, Inc. All Rights Reserved.
+
+#pragma once
+
+#include "CoreMinimal.h"
+#include "SceneOutlinerFwd.h"
+#include "ISceneOutliner.h"
+#include "ISceneOutlinerColumn.h"
+#include "Widgets/Views/SHeaderRow.h"
+
+template<typename ItemType> class STableRow;
+class FDataLayerOutlinerIsDynamicallyLoadedColumn : public ISceneOutlinerColumn
+{
+public:
+	FDataLayerOutlinerIsDynamicallyLoadedColumn(ISceneOutliner& SceneOutliner) : WeakSceneOutliner(StaticCastSharedRef<ISceneOutliner>(SceneOutliner.AsShared())) {}
+	virtual ~FDataLayerOutlinerIsDynamicallyLoadedColumn() {}
+	static FName GetID();
+
+	//////////////////////////////////////////////////////////////////////////
+	// Begin ISceneOutlinerColumn Implementation
+	virtual FName GetColumnID() override { return GetID(); }
+	virtual SHeaderRow::FColumn::FArguments ConstructHeaderRowColumn() override;
+	virtual const TSharedRef<SWidget> ConstructRowWidget(FSceneOutlinerTreeItemRef TreeItem, const STableRow<FSceneOutlinerTreeItemPtr>& Row) override;
+	virtual bool SupportsSorting() const override { return false; }
+	// End ISceneOutlinerColumn Implementation
+	//////////////////////////////////////////////////////////////////////////
+private:
+	TWeakPtr<ISceneOutliner> WeakSceneOutliner;
+};
\ No newline at end of file
diff --git a/Engine/Source/Editor/DataLayerEditor/Private/DataLayer/DataLayerPropertyTypeCustomization.cpp b/Engine/Source/Editor/DataLayerEditor/Private/DataLayer/DataLayerPropertyTypeCustomization.cpp
new file mode 100644
--- /dev/null
+++ b/Engine/Source/Editor/DataLayerEditor/Private/DataLayer/DataLayerPropertyTypeCustomization.cpp
@@ -0,0 +1,217 @@
+// Copyright Epic Games, Inc. All Rights Reserved.
+
+#include "DataLayerPropertyTypeCustomization.h"
+#include "DataLayer/DataLayerDragDropOp.h"
+#include "DataLayer/DataLayerEditorSubsystem.h"
+#include "WorldPartition/DataLayer/DataLayer.h"
+#include "Algo/Accumulate.h"
+#include "Modules/ModuleManager.h"
+#include "PropertyHandle.h"
+#include "DetailWidgetRow.h"
+#include "Widgets/SWidget.h"
+#include "Widgets/SBoxPanel.h"
+#include "Widgets/Input/SButton.h"
+#include "Widgets/Input/SComboButton.h"
+#include "Widgets/Text/STextBlock.h"
+#include "Framework/MultiBox/MultiBoxBuilder.h"
+#include "SDropTarget.h"
+#include "Editor.h"
+#include "EditorFontGlyphs.h"
+#include "LevelEditor.h"
+#include "DetailLayoutBuilder.h"
+#include "DetailCategoryBuilder.h"
+
+#define LOCTEXT_NAMESPACE "DataLayer"
+
+void FDataLayerPropertyTypeCustomization::CustomizeHeader(TSharedRef<IPropertyHandle> StructPropertyHandle, FDetailWidgetRow& HeaderRow, IPropertyTypeCustomizationUtils& StructCustomizationUtils)
+{
+	PropertyHandle = StructPropertyHandle->GetChildHandle("Name");
+
+	HeaderRow
+	.NameContent()
+	[
+		StructPropertyHandle->CreatePropertyNameWidget()
+	]
+	.ValueContent()
+	.MaxDesiredWidth(TOptional<float>())
+	[
+		SNew(SDropTarget)
+		.OnDrop(this, &FDataLayerPropertyTypeCustomization::OnDrop)
+		.OnAllowDrop(this, &FDataLayerPropertyTypeCustomization::OnVerifyDrag)
+		.OnIsRecognized(this, &FDataLayerPropertyTypeCustomization::OnVerifyDrag)
+		[
+			SNew(SHorizontalBox)
+
+			+ SHorizontalBox::Slot()
+			.VAlign(VAlign_Center)
+			.AutoWidth()
+			[
+				SNew(SImage)
+				.Image(FEditorStyle::GetBrush(TEXT("DataLayer.Icon16x")))
+				.ColorAndOpacity(FSlateColor::UseForeground())
+			]
+
+			+ SHorizontalBox::Slot()
+			.VAlign(VAlign_Center)
+			.Padding(3.0f, 0.0f)
+			.FillWidth(1.0f)
+			[
+				SNew(SComboButton)
+				.ToolTipText(LOCTEXT("ComboButtonTip", "Drag and drop a DataLayer onto this property, or choose one from the drop down."))
+				.OnGetMenuContent(this, &FDataLayerPropertyTypeCustomization::OnGetDataLayerMenu)
+				.ButtonStyle(FEditorStyle::Get(), "NoBorder")
+				.ForegroundColor(FSlateColor::UseForeground())
+				.ContentPadding(FMargin(0))
+				.ButtonContent()
+				[
+					SNew(STextBlock)
+					.Text(this, &FDataLayerPropertyTypeCustomization::GetDataLayerText)
+				]
+			]
+
+			+ SHorizontalBox::Slot()
+			.VAlign(VAlign_Center)
+			.Padding(1.0f, 0.0f)
+			.AutoWidth()
+			[
+				SNew(SButton)
+				.ButtonStyle(FEditorStyle::Get(), "HoverHintOnly")
+				.ToolTipText(LOCTEXT("SelectTip", "Select all actors in this DataLayer"))
+				.OnClicked(this, &FDataLayerPropertyTypeCustomization::OnSelectDataLayer)
+				.Visibility(this, &FDataLayerPropertyTypeCustomization::GetSelectDataLayerVisibility)
+				.ForegroundColor(FSlateColor::UseForeground())
+				[
+					SNew(STextBlock)
+					.Font(FEditorStyle::Get().GetFontStyle("FontAwesome.11"))
+					.Text(FEditorFontGlyphs::Sign_In)
+				]
+			]
+		]
+	];
+}
+
+FText FDataLayerPropertyTypeCustomization::GetDataLayerDescription(const UDataLayer* InDataLayer) const
+{
+	check(InDataLayer);
+	return FText::FromName(InDataLayer->GetDataLayerLabel());
+}
+
+const UDataLayer* FDataLayerPropertyTypeCustomization::GetDataLayerFromPropertyHandle() const
+{
+	FName DataLayerName;
+	if (PropertyHandle->GetValue(DataLayerName) == FPropertyAccess::Success)
+	{
+		UDataLayerEditorSubsystem* DataLayerSubsystem = UDataLayerEditorSubsystem::Get();
+		const UDataLayer* DataLayer = DataLayerSubsystem->GetDataLayerFromName(DataLayerName);
+		return DataLayer;
+	}
+	return nullptr;
+}
+
+FText FDataLayerPropertyTypeCustomization::GetDataLayerText() const
+{
+	const UDataLayer* DataLayer = GetDataLayerFromPropertyHandle();
+	return DataLayer ? GetDataLayerDescription(DataLayer) : LOCTEXT("InvalidDataLayerLabel", "<Invalid>");
+}
+
+TSharedRef<SWidget> FDataLayerPropertyTypeCustomization::OnGetDataLayerMenu()
+{
+	FMenuBuilder MenuBuilder(true, nullptr);
+
+	MenuBuilder.AddMenuEntry(
+		LOCTEXT("OpenDataLayersBrowser", "Browse DataLayers..."),
+		FText(),
+		FSlateIcon(FEditorStyle::GetStyleSetName(), "LevelEditor.Tabs.DataLayers"),
+		FUIAction(
+			FExecuteAction::CreateSP(this, &FDataLayerPropertyTypeCustomization::OpenDataLayerBrowser)
+		)
+	);
+
+	MenuBuilder.BeginSection(FName(), LOCTEXT("ExistingDataLayers", "Existing DataLayers"));
+	{
+		TArray<TWeakObjectPtr<UDataLayer>> AllDataLayers;
+		UDataLayerEditorSubsystem* DataLayerSubsystem = UDataLayerEditorSubsystem::Get();
+		DataLayerSubsystem->AddAllDataLayersTo(AllDataLayers);
+
+		for (const TWeakObjectPtr<UDataLayer>& WeakDataLayer : AllDataLayers)
+		{
+			if (const UDataLayer* DataLayerPtr = WeakDataLayer.Get())
+			{
+				MenuBuilder.AddMenuEntry(
+					GetDataLayerDescription(DataLayerPtr),
+					FText(),
+					FSlateIcon(FEditorStyle::GetStyleSetName(), "DataLayer.Icon16x"),
+					FUIAction(
+						FExecuteAction::CreateSP(this, &FDataLayerPropertyTypeCustomization::AssignDataLayer, DataLayerPtr)
+					)
+				);
+			}
+		}
+	}
+	MenuBuilder.EndSection();
+
+	return MenuBuilder.MakeWidget();
+}
+
+EVisibility FDataLayerPropertyTypeCustomization::GetSelectDataLayerVisibility() const
+{
+	const UDataLayer* DataLayer = GetDataLayerFromPropertyHandle();
+	return DataLayer ? EVisibility::Visible : EVisibility::Collapsed;
+}
+
+FReply FDataLayerPropertyTypeCustomization::OnSelectDataLayer()
+{
+	if (const UDataLayer* DataLayer = GetDataLayerFromPropertyHandle())
+	{
+		GEditor->SelectNone(true, true);
+		UDataLayerEditorSubsystem::Get()->SelectActorsInDataLayer(DataLayer, true, true);
+	}
+	return FReply::Handled();
+}
+
+void FDataLayerPropertyTypeCustomization::AssignDataLayer(const UDataLayer* InDataLayer)
+{
+	if (GetDataLayerFromPropertyHandle() != InDataLayer)
+	{
+		PropertyHandle->SetValue(InDataLayer ? InDataLayer->GetFName() : NAME_None);
+		UDataLayerEditorSubsystem::Get()->OnDataLayerChanged().Broadcast(EDataLayerAction::Reset, NULL, NAME_None);
+	}
+}
+
+void FDataLayerPropertyTypeCustomization::OpenDataLayerBrowser()
+{
+	FLevelEditorModule& LevelEditorModule = FModuleManager::GetModuleChecked<FLevelEditorModule>("LevelEditor");
+	LevelEditorModule.GetLevelEditorTabManager()->TryInvokeTab(FTabId("LevelEditorDataLayerBrowser"));
+}
+
+FReply FDataLayerPropertyTypeCustomization::OnDrop(TSharedPtr<FDragDropOperation> InDragDrop)
+{
+	if (InDragDrop.IsValid() && InDragDrop->IsOfType<FDataLayerDragDropOp>())
+	{
+		TSharedPtr<const FDataLayerDragDropOp> DataLayerDragDropOp = InDragDrop->CastTo<FDataLayerDragDropOp>();
+		if (DataLayerDragDropOp.IsValid())
+		{
+			const TArray<FName>& DataLayerLabels = DataLayerDragDropOp->DataLayerLabels;
+			if (ensure(DataLayerLabels.Num() == 1))
+			{
+				if (const UDataLayer* DataLayerPtr = UDataLayerEditorSubsystem::Get()->GetDataLayerFromLabel(DataLayerLabels[0]))
+				{
+					AssignDataLayer(DataLayerPtr);
+				}
+			}
+		}
+	}
+	return FReply::Handled();
+}
+
+bool FDataLayerPropertyTypeCustomization::OnVerifyDrag(TSharedPtr<FDragDropOperation> InDragDrop)
+{
+	if (InDragDrop.IsValid() && InDragDrop->IsOfType<FDataLayerDragDropOp>())
+	{
+		TSharedPtr<const FDataLayerDragDropOp> DataLayerDragDropOp = InDragDrop->CastTo<FDataLayerDragDropOp>();
+		return DataLayerDragDropOp.IsValid() && DataLayerDragDropOp->DataLayerLabels.Num() == 1;
+	}
+	return false;
+}
+
+#undef LOCTEXT_NAMESPACE
\ No newline at end of file
diff --git a/Engine/Source/Editor/DataLayerEditor/Private/DataLayer/DataLayerPropertyTypeCustomization.h b/Engine/Source/Editor/DataLayerEditor/Private/DataLayer/DataLayerPropertyTypeCustomization.h
new file mode 100644
--- /dev/null
+++ b/Engine/Source/Editor/DataLayerEditor/Private/DataLayer/DataLayerPropertyTypeCustomization.h
@@ -0,0 +1,36 @@
+// Copyright Epic Games, Inc. All Rights Reserved.
+
+#pragma once
+
+#include "CoreTypes.h"
+#include "IPropertyTypeCustomization.h"
+
+struct EVisibility;
+class FReply;
+class SWidget;
+class FDragDropOperation;
+class UDataLayer;
+
+struct FDataLayerPropertyTypeCustomization : public IPropertyTypeCustomization
+{
+	virtual void CustomizeHeader(TSharedRef<IPropertyHandle> StructPropertyHandle, FDetailWidgetRow& HeaderRow, IPropertyTypeCustomizationUtils& StructCustomizationUtils) override;
+	virtual void CustomizeChildren(TSharedRef<IPropertyHandle> StructPropertyHandle, IDetailChildrenBuilder& StructBuilder, IPropertyTypeCustomizationUtils& StructCustomizationUtils) override {}
+
+private:
+
+	void AssignDataLayer(const UDataLayer* InDataLayer);
+	void OpenDataLayerBrowser();
+
+	const UDataLayer* GetDataLayerFromPropertyHandle() const;
+	FText GetDataLayerDescription(const UDataLayer* InDataLayer) const;
+	FText GetDataLayerText() const;
+	EVisibility GetSelectDataLayerVisibility() const;
+
+	TSharedRef<SWidget> OnGetDataLayerMenu();
+	FReply OnSelectDataLayer();
+	FReply OnDrop(TSharedPtr<FDragDropOperation> InDragDrop);
+	bool OnVerifyDrag(TSharedPtr<FDragDropOperation> InDragDrop);
+
+	TSharedPtr<IPropertyHandle> PropertyHandle;
+};
+
diff --git a/Engine/Source/Editor/DataLayerEditor/Private/DataLayer/DataLayerTreeItem.cpp b/Engine/Source/Editor/DataLayerEditor/Private/DataLayer/DataLayerTreeItem.cpp
new file mode 100644
--- /dev/null
+++ b/Engine/Source/Editor/DataLayerEditor/Private/DataLayer/DataLayerTreeItem.cpp
@@ -0,0 +1,47 @@
+// Copyright Epic Games, Inc. All Rights Reserved.
+
+#include "DataLayerTreeItem.h"
+#include "DataLayer/DataLayerEditorSubsystem.h"
+#include "ISceneOutlinerTreeItem.h"
+#include "SceneOutlinerStandaloneTypes.h"
+#include "SDataLayerTreeLabel.h"
+#include "Editor.h"
+
+#define LOCTEXT_NAMESPACE "DataLayer"
+
+const FSceneOutlinerTreeItemType FDataLayerTreeItem::Type(&ISceneOutlinerTreeItem::Type);
+
+FDataLayerTreeItem::FDataLayerTreeItem(UDataLayer* InDataLayer)
+	: ISceneOutlinerTreeItem(Type)
+	, DataLayer(InDataLayer)
+	, ID(InDataLayer)
+{
+	Flags.bIsExpanded = false;
+}
+
+FString FDataLayerTreeItem::GetDisplayString() const
+{
+	const UDataLayer* DataLayerPtr = DataLayer.Get();
+	return DataLayerPtr ? DataLayerPtr->GetDataLayerLabel().ToString() : LOCTEXT("DataLayerForMissingDataLayer", "(Deleted DataLayer)").ToString();
+}
+
+bool FDataLayerTreeItem::GetVisibility() const
+{
+	const UDataLayer* DataLayerPtr = DataLayer.Get();
+	return DataLayerPtr && DataLayerPtr->IsVisible();
+}
+
+TSharedRef<SWidget> FDataLayerTreeItem::GenerateLabelWidget(ISceneOutliner& Outliner, const STableRow<FSceneOutlinerTreeItemPtr>& InRow)
+{
+	return SNew(SDataLayerTreeLabel, *this, Outliner, InRow);
+}
+
+void FDataLayerTreeItem::OnVisibilityChanged(const bool bNewVisibility)
+{
+	if (UDataLayer* DataLayerPtr = DataLayer.Get())
+	{
+		UDataLayerEditorSubsystem::Get()->SetDataLayerVisibility(DataLayerPtr, bNewVisibility);
+	}
+}
+
+#undef LOCTEXT_NAMESPACE
\ No newline at end of file
diff --git a/Engine/Source/Editor/DataLayerEditor/Private/DataLayer/DataLayerTreeItem.h b/Engine/Source/Editor/DataLayerEditor/Private/DataLayer/DataLayerTreeItem.h
new file mode 100644
--- /dev/null
+++ b/Engine/Source/Editor/DataLayerEditor/Private/DataLayer/DataLayerTreeItem.h
@@ -0,0 +1,33 @@
+// Copyright Epic Games, Inc. All Rights Reserved.
+
+#pragma once
+
+#include "CoreMinimal.h"
+#include "Misc/Attribute.h"
+#include "SceneOutlinerFwd.h"
+#include "ISceneOutlinerTreeItem.h"
+#include "WorldPartition/DataLayer/DataLayer.h"
+
+struct FDataLayerTreeItem : ISceneOutlinerTreeItem
+{
+	FDataLayerTreeItem(UDataLayer* InDataLayer);
+	const UDataLayer* GetDataLayer() const { return DataLayer.Get(); }
+	UDataLayer* GetDataLayer() { return DataLayer.Get(); }
+
+	/* Begin ISceneOutlinerTreeItem Implementation */
+	virtual bool IsValid() const override { return DataLayer.IsValid(); }
+	virtual FSceneOutlinerTreeItemID GetID() const override { return ID; }
+	virtual FString GetDisplayString() const override;
+	virtual bool CanInteract() const override { return true; }
+	virtual TSharedRef<SWidget> GenerateLabelWidget(ISceneOutliner& Outliner, const STableRow<FSceneOutlinerTreeItemPtr>& InRow) override;
+	virtual bool HasVisibilityInfo() const override { return true; }
+	virtual void OnVisibilityChanged(const bool bNewVisibility) override;
+	virtual bool GetVisibility() const override;
+	/* End ISceneOutlinerTreeItem Implementation */
+
+	static const FSceneOutlinerTreeItemType Type;
+
+private:
+	TWeakObjectPtr<UDataLayer> DataLayer;
+	const FObjectKey ID;
+};
\ No newline at end of file
diff --git a/Engine/Source/Editor/DataLayerEditor/Private/DataLayer/SDataLayerBrowser.cpp b/Engine/Source/Editor/DataLayerEditor/Private/DataLayer/SDataLayerBrowser.cpp
new file mode 100644
--- /dev/null
+++ b/Engine/Source/Editor/DataLayerEditor/Private/DataLayer/SDataLayerBrowser.cpp
@@ -0,0 +1,159 @@
+// Copyright Epic Games, Inc. All Rights Reserved.
+
+#include "SDataLayerBrowser.h"
+#include "Framework/Application/SlateApplication.h"
+#include "Widgets/Input/SButton.h"
+#include "SDataLayerOutliner.h"
+#include "DataLayerActorTreeItem.h"
+#include "DataLayerTreeItem.h"
+#include "DataLayerMode.h"
+#include "DataLayerModeSettings.h"
+#include "DataLayerOutlinerIsDynamicallyLoadedColumn.h"
+#include "DataLayerOutlinerDeleteButtonColumn.h"
+#include "DataLayer/DataLayerEditorSubsystem.h"
+#include "WorldPartition/DataLayer/DataLayer.h"
+#include "ScopedTransaction.h"
+#include "Editor.h"
+
+#define LOCTEXT_NAMESPACE "DataLayer"
+
+void SDataLayerBrowser::Construct(const FArguments& InArgs)
+{
+	const UDataLayerModeSettings* SharedSettings = GetMutableDefault<UDataLayerModeSettings>();
+	Mode = SharedSettings->bShowDataLayerContent ? EDataLayerBrowserMode::DataLayerContents : EDataLayerBrowserMode::DataLayers;
+
+	auto ToggleDataLayerContents = [this]()
+	{
+		const EDataLayerBrowserMode NewMode = Mode == EDataLayerBrowserMode::DataLayers ? EDataLayerBrowserMode::DataLayerContents : EDataLayerBrowserMode::DataLayers;
+		SetupDataLayerMode(NewMode);
+		return FReply::Handled();
+	};
+
+	auto GetToggleModeButtonImageBrush = [this]()
+	{
+		static const FName ExploreDataLayerContents("DataLayerBrowser.ExploreDataLayerContents");
+		static const FName ReturnToDataLayersList("DataLayerBrowser.ReturnToDataLayersList");
+		return (Mode == EDataLayerBrowserMode::DataLayers) ? FEditorStyle::GetBrush(ExploreDataLayerContents) : FEditorStyle::GetBrush(ReturnToDataLayersList);
+	};
+
+	auto GetToggleModeButtonText = [this]()
+	{
+		return (Mode == EDataLayerBrowserMode::DataLayers) ? LOCTEXT("SeeContentsLabel", "See Contents") : LOCTEXT("HideContentsLabel", "Hide Contents");
+	};
+
+	auto GetInvertedForegroundIfHovered = [this]()
+	{
+		static const FName InvertedForegroundName("InvertedForeground");
+		return (ToggleModeButton.IsValid() && (ToggleModeButton->IsHovered() || ToggleModeButton->IsPressed())) ? FEditorStyle::GetSlateColor(InvertedForegroundName) : FSlateColor::UseForeground();
+	};
+
+	//////////////////////////////////////////////////////////////////////////
+	//	DataLayer Contents Header
+	SAssignNew(DataLayerContentsHeader, SBorder)
+	.BorderImage(FEditorStyle::GetBrush("DataLayerBrowser.DataLayerContentsQuickbarBackground"))
+	.Visibility(EVisibility::Visible)
+	.Content()
+	[
+		SNew(SHorizontalBox)
+		+ SHorizontalBox::Slot()
+		.AutoWidth()
+		.Padding(FMargin(0, 0, 2, 0))
+		[
+			SAssignNew(ToggleModeButton, SButton)
+			.ContentPadding(FMargin(2, 0, 2, 0))
+			.ButtonStyle(FEditorStyle::Get(), "DataLayerBrowserButton")
+			.OnClicked_Lambda(ToggleDataLayerContents)
+			.ForegroundColor(FSlateColor::UseForeground())
+			.VAlign(VAlign_Center)
+			.HAlign(HAlign_Center)
+			.Content()
+			[
+				SNew(SHorizontalBox)
+				+ SHorizontalBox::Slot()
+				.AutoWidth()
+				.HAlign(HAlign_Center)
+				.VAlign(VAlign_Center)
+				.Padding(0, 1, 3, 1)
+				[
+					SNew(SImage)
+					.Image_Lambda(GetToggleModeButtonImageBrush)
+					.ColorAndOpacity_Lambda(GetInvertedForegroundIfHovered)
+				]
+
+				+ SHorizontalBox::Slot()
+				.AutoWidth()
+				.HAlign(HAlign_Center)
+				.VAlign(VAlign_Center)
+				[
+					SNew(STextBlock)
+					.Text_Lambda(GetToggleModeButtonText)
+					.ColorAndOpacity_Lambda(GetInvertedForegroundIfHovered)
+				]
+			]
+		]
+	];
+
+	//////////////////////////////////////////////////////////////////////////
+	//	DataLayer Contents Section
+	FSceneOutlinerInitializationOptions InitOptions;
+	InitOptions.bShowHeaderRow = false;
+	InitOptions.bShowParentTree = true;
+	InitOptions.bShowCreateNewFolder = false;
+	InitOptions.ModeFactory = FCreateSceneOutlinerMode::CreateLambda([this](SSceneOutliner* Outliner) { return new FDataLayerMode(FDataLayerModeParams(Outliner, this, nullptr)); });
+	InitOptions.ColumnMap.Add(FSceneOutlinerBuiltInColumnTypes::Gutter(), FSceneOutlinerColumnInfo(ESceneOutlinerColumnVisibility::Visible, 0));
+	InitOptions.ColumnMap.Add(FDataLayerOutlinerIsDynamicallyLoadedColumn::GetID(), FSceneOutlinerColumnInfo(ESceneOutlinerColumnVisibility::Visible, 1, FCreateSceneOutlinerColumn::CreateLambda([](ISceneOutliner& InSceneOutliner) { return MakeShareable(new FDataLayerOutlinerIsDynamicallyLoadedColumn(InSceneOutliner)); })));
+	InitOptions.ColumnMap.Add(FSceneOutlinerBuiltInColumnTypes::Label(), FSceneOutlinerColumnInfo(ESceneOutlinerColumnVisibility::Visible, 2));
+	InitOptions.ColumnMap.Add(FDataLayerOutlinerDeleteButtonColumn::GetID(), FSceneOutlinerColumnInfo(ESceneOutlinerColumnVisibility::Visible, 20, FCreateSceneOutlinerColumn::CreateLambda([](ISceneOutliner& InSceneOutliner) { return MakeShareable(new FDataLayerOutlinerDeleteButtonColumn(InSceneOutliner)); })));
+
+	auto CreateDataLayerBrowser = [&](const FSceneOutlinerInitializationOptions& InInitOptions)
+	{
+		return SNew(SDataLayerOutliner, InInitOptions).IsEnabled(FSlateApplication::Get().GetNormalExecutionAttribute());
+	};
+
+	SAssignNew(DataLayerContentsSection, SBorder)
+	.Padding(5)
+	.BorderImage(FEditorStyle::GetBrush("NoBrush"))
+	.Content()
+	[
+		CreateDataLayerBrowser(InitOptions)
+	];
+
+	//////////////////////////////////////////////////////////////////////////
+	//	DataLayer Browser
+	ChildSlot
+	[
+		SAssignNew(ContentAreaBox, SVerticalBox)
+	];
+
+	SetupDataLayerMode(Mode);
+}
+
+
+void SDataLayerBrowser::SetupDataLayerMode(EDataLayerBrowserMode InNewMode)
+{
+	ContentAreaBox->ClearChildren();
+	ContentAreaBox->AddSlot()
+	.AutoHeight()
+	.FillHeight(1.0f)
+	[
+		DataLayerContentsSection.ToSharedRef()
+	];
+
+	ContentAreaBox->AddSlot()
+	.AutoHeight()
+	.VAlign(VAlign_Bottom)
+	.MaxHeight(23)
+	[
+		DataLayerContentsHeader.ToSharedRef()
+	];
+
+	Mode = InNewMode;
+	
+	UDataLayerModeSettings* SharedSettings = GetMutableDefault<UDataLayerModeSettings>();
+	SharedSettings->bShowDataLayerContent = Mode == EDataLayerBrowserMode::DataLayerContents;
+	SharedSettings->PostEditChange();
+
+	ModeChanged.Broadcast(Mode);
+}
+
+#undef LOCTEXT_NAMESPACE
\ No newline at end of file
diff --git a/Engine/Source/Editor/DataLayerEditor/Private/DataLayer/SDataLayerBrowser.h b/Engine/Source/Editor/DataLayerEditor/Private/DataLayer/SDataLayerBrowser.h
new file mode 100644
--- /dev/null
+++ b/Engine/Source/Editor/DataLayerEditor/Private/DataLayer/SDataLayerBrowser.h
@@ -0,0 +1,62 @@
+// Copyright Epic Games, Inc. All Rights Reserved.
+#pragma once
+
+#include "CoreMinimal.h"
+#include "InputCoreTypes.h"
+#include "Styling/SlateColor.h"
+#include "Input/Reply.h"
+#include "Widgets/SWidget.h"
+#include "Widgets/DeclarativeSyntaxSupport.h"
+#include "Widgets/SCompoundWidget.h"
+
+class ISceneOutliner;
+class SButton;
+class SVerticalBox;
+class SBorder;
+
+enum class EDataLayerBrowserMode
+{
+	DataLayers,
+	DataLayerContents,
+	Count
+};
+
+//////////////////////////////////////////////////////////////////////////
+// SDataLayerBrowser
+
+class SDataLayerBrowser : public SCompoundWidget
+{
+public:
+
+	SLATE_BEGIN_ARGS( SDataLayerBrowser ) {}
+	SLATE_END_ARGS()
+
+	~SDataLayerBrowser(){}
+
+	/**
+	 * Construct this widget.  Called by the SNew() Slate macro.
+	 *
+	 * @param	InArgs				Declaration used by the SNew() macro to construct this widget
+	 * @param	InViewModel			The UI logic not specific to slate
+	 */
+	void Construct(const FArguments& InArgs);
+
+	EDataLayerBrowserMode GetMode() const { return Mode; }
+
+	/** Broadcasts whenever one or more Actors changed DataLayers*/
+	DECLARE_EVENT_OneParam(SDataLayerBrowser, FOnModeChanged, EDataLayerBrowserMode /*Mode*/);
+	virtual FOnModeChanged& OnModeChanged() final { return ModeChanged; }
+
+private:
+
+	void SetupDataLayerMode(EDataLayerBrowserMode InNewMode);
+	
+	/**	Fires whenever one or more actor DataLayer changes */
+	FOnModeChanged ModeChanged;
+
+	TSharedPtr<SButton> ToggleModeButton;
+	TSharedPtr<SVerticalBox> ContentAreaBox;
+	TSharedPtr<SBorder> DataLayerContentsSection;
+	TSharedPtr<SBorder> DataLayerContentsHeader;
+	EDataLayerBrowserMode Mode;
+};
diff --git a/Engine/Source/Editor/DataLayerEditor/Private/DataLayer/SDataLayerOutliner.h b/Engine/Source/Editor/DataLayerEditor/Private/DataLayer/SDataLayerOutliner.h
new file mode 100644
--- /dev/null
+++ b/Engine/Source/Editor/DataLayerEditor/Private/DataLayer/SDataLayerOutliner.h
@@ -0,0 +1,18 @@
+// Copyright Epic Games, Inc. All Rights Reserved.
+
+#pragma once
+
+#include "CoreMinimal.h"
+#include "SSceneOutliner.h"
+
+class SDataLayerOutliner : public SSceneOutliner
+{
+public:
+	void Construct(const FArguments& InArgs, const FSceneOutlinerInitializationOptions& InitOptions)
+	{
+		SSceneOutliner::Construct(InArgs, InitOptions);
+	}
+
+	SDataLayerOutliner() {}
+	virtual ~SDataLayerOutliner() {}
+};
\ No newline at end of file
diff --git a/Engine/Source/Editor/DataLayerEditor/Private/DataLayer/SDataLayerTreeLabel.cpp b/Engine/Source/Editor/DataLayerEditor/Private/DataLayer/SDataLayerTreeLabel.cpp
new file mode 100644
--- /dev/null
+++ b/Engine/Source/Editor/DataLayerEditor/Private/DataLayer/SDataLayerTreeLabel.cpp
@@ -0,0 +1,184 @@
+// Copyright Epic Games, Inc. All Rights Reserved.
+
+#include "SDataLayerTreeLabel.h"
+#include "DataLayer/DataLayerEditorSubsystem.h"
+#include "DataLayerTreeItem.h"
+#include "Widgets/Text/SInlineEditableTextBlock.h"
+#include "ISceneOutliner.h"
+#include "ISceneOutlinerMode.h"
+#include "ScopedTransaction.h"
+#include "Editor.h"
+
+#define LOCTEXT_NAMESPACE "DataLayer"
+
+void SDataLayerTreeLabel::Construct(const FArguments& InArgs, FDataLayerTreeItem& DataLayerItem, ISceneOutliner& SceneOutliner, const STableRow<FSceneOutlinerTreeItemPtr>& InRow)
+{
+	WeakSceneOutliner = StaticCastSharedRef<ISceneOutliner>(SceneOutliner.AsShared());
+	TreeItemPtr = StaticCastSharedRef<FDataLayerTreeItem>(DataLayerItem.AsShared());
+	DataLayerPtr = DataLayerItem.GetDataLayer();
+	HighlightText = SceneOutliner.GetFilterHighlightText();
+
+	TSharedPtr<SInlineEditableTextBlock> InlineTextBlock;
+
+	auto MainContent = SNew(SHorizontalBox)
+
+	+ SHorizontalBox::Slot()
+	.VAlign(VAlign_Center)
+	[
+		SAssignNew(InlineTextBlock, SInlineEditableTextBlock)
+		.Text(this, &SDataLayerTreeLabel::GetDisplayText)
+		.ToolTipText(this, &SDataLayerTreeLabel::GetTooltipText)
+		.HighlightText(HighlightText)
+		.ColorAndOpacity(this, &SDataLayerTreeLabel::GetForegroundColor)
+		.OnTextCommitted(this, &SDataLayerTreeLabel::OnLabelCommitted)
+		.OnVerifyTextChanged(this, &SDataLayerTreeLabel::OnVerifyItemLabelChanged)
+		.IsSelected(FIsSelected::CreateSP(&InRow, &STableRow<FSceneOutlinerTreeItemPtr>::IsSelectedExclusively))
+		.IsReadOnly_Lambda([Item = DataLayerItem.AsShared(), this]()
+		{
+			return !CanExecuteRenameRequest(Item.Get());
+		})
+	]
+
+	+ SHorizontalBox::Slot()
+	.VAlign(VAlign_Center)
+	.AutoWidth()
+	.Padding(0.0f, 0.f, 3.0f, 0.0f)
+	[
+		SNew(STextBlock)
+		.Text(this, &SDataLayerTreeLabel::GetTypeText)
+		.Visibility(this, &SDataLayerTreeLabel::GetTypeTextVisibility)
+		.HighlightText(HighlightText)
+	];
+
+	if (WeakSceneOutliner.Pin()->GetMode()->IsInteractive())
+	{
+		DataLayerItem.RenameRequestEvent.BindSP(InlineTextBlock.Get(), &SInlineEditableTextBlock::EnterEditingMode);
+	}
+
+	ChildSlot
+	[
+		SNew(SHorizontalBox)
+		+ SHorizontalBox::Slot()
+		.AutoWidth()
+		.VAlign(VAlign_Center)
+		.Padding(FSceneOutlinerDefaultTreeItemMetrics::IconPadding())
+		[
+			SNew(SBox)
+			.WidthOverride(FSceneOutlinerDefaultTreeItemMetrics::IconSize())
+			.HeightOverride(FSceneOutlinerDefaultTreeItemMetrics::IconSize())
+			[
+				SNew(SImage)
+				.Image(this, &SDataLayerTreeLabel::GetIcon)
+				.ToolTipText(this, &SDataLayerTreeLabel::GetIconTooltip)
+				.ColorAndOpacity(FSlateColor::UseForeground())
+			]
+		]
+
+		+ SHorizontalBox::Slot()
+		.FillWidth(1.0f)
+		.VAlign(VAlign_Center)
+		.Padding(0.0f, 0.0f)
+		[
+			MainContent
+		]
+	];
+}
+
+FText SDataLayerTreeLabel::GetDisplayText() const
+{
+	const UDataLayer* DataLayer = DataLayerPtr.Get();
+	return DataLayer ? FText::FromName(DataLayer->GetDataLayerLabel()) : LOCTEXT("DataLayerLabelForMissingDataLayer", "(Deleted Data Layer)");
+}
+
+FText SDataLayerTreeLabel::GetTooltipText() const
+{
+	if (const UDataLayer* DataLayer = DataLayerPtr.Get())
+	{
+		FFormatNamedArguments Args;
+		Args.Add(TEXT("ID_Name"), LOCTEXT("CustomColumnMode_InternalName", "ID Name"));
+		Args.Add(TEXT("Name"), FText::FromName(DataLayer->GetFName()));
+		return FText::Format(LOCTEXT("DataLayerNameTooltip", "{ID_Name}: {Name}"), Args);
+	}
+	return FText();
+}
+
+FText SDataLayerTreeLabel::GetTypeText() const
+{
+	const UDataLayer* DataLayer = DataLayerPtr.Get();
+	return DataLayer ? FText::FromName(DataLayer->GetClass()->GetFName()) : FText();
+}
+
+EVisibility SDataLayerTreeLabel::GetTypeTextVisibility() const
+{
+	return HighlightText.Get().IsEmpty() ? EVisibility::Collapsed : EVisibility::Visible;
+}
+
+const FSlateBrush* SDataLayerTreeLabel::GetIcon() const
+{
+	const UDataLayer* DataLayer = DataLayerPtr.Get();
+	if (DataLayer && WeakSceneOutliner.IsValid())
+	{
+		const FName IconName = DataLayer->GetClass()->GetFName();
+		if (const FSlateBrush* CachedBrush = WeakSceneOutliner.Pin()->GetCachedIconForClass(IconName))
+		{
+			return CachedBrush;
+		}
+
+		const FSlateBrush* FoundSlateBrush = FEditorStyle::GetBrush(TEXT("DataLayer.Icon16x"));
+		WeakSceneOutliner.Pin()->CacheIconForClass(IconName, FoundSlateBrush);
+		return FoundSlateBrush;
+	}
+	return nullptr;
+}
+
+FText SDataLayerTreeLabel::GetIconTooltip() const
+{
+	return TreeItemPtr.Pin().IsValid() && DataLayerPtr.Get() ? FText::FromString(DataLayerPtr.Get()->GetClass()->GetName()) : FText();
+}
+
+FSlateColor SDataLayerTreeLabel::GetForegroundColor() const
+{
+	if (TOptional<FLinearColor> BaseColor = FSceneOutlinerCommonLabelData::GetForegroundColor(*TreeItemPtr.Pin()))
+	{
+		return BaseColor.GetValue();
+	}
+
+	const UDataLayer* DataLayer = DataLayerPtr.Get();
+	return (!DataLayer || !DataLayer->GetWorld()) ? FLinearColor(0.2f, 0.2f, 0.25f) : FSlateColor::UseForeground();
+}
+
+bool SDataLayerTreeLabel::OnVerifyItemLabelChanged(const FText& InLabel, FText& OutErrorMessage)
+{
+	if (InLabel.IsEmptyOrWhitespace())
+	{
+		OutErrorMessage = LOCTEXT("EmptyDataLayerLabel", "DataLayer must be given a name");
+		return false;
+	}
+
+	const UDataLayer* FoundDataLayer;
+	if (UDataLayerEditorSubsystem::Get()->TryGetDataLayerFromLabel(*InLabel.ToString(), FoundDataLayer) && FoundDataLayer != DataLayerPtr.Get())
+	{
+		OutErrorMessage = LOCTEXT("RenameFailed_AlreadyExists", "This DataLayer already exists");
+		return false;
+	}
+
+	return true;
+}
+
+void SDataLayerTreeLabel::OnLabelCommitted(const FText& InLabel, ETextCommit::Type InCommitInfo)
+{
+	UDataLayer* DataLayer = DataLayerPtr.Get();
+	if (DataLayer && !InLabel.ToString().Equals(DataLayer->GetDataLayerLabel().ToString(), ESearchCase::CaseSensitive))
+	{
+		const FScopedTransaction Transaction(LOCTEXT("SceneOutlinerRenameDataLayerTransaction", "Rename Data Layer"));
+		UDataLayerEditorSubsystem::Get()->RenameDataLayer(DataLayer, *InLabel.ToString());
+		DataLayer->SetDataLayerLabel(*InLabel.ToString());
+
+		if (WeakSceneOutliner.IsValid())
+		{
+			WeakSceneOutliner.Pin()->SetKeyboardFocus();
+		}
+	}
+}
+
+#undef LOCTEXT_NAMESPACE 
\ No newline at end of file
diff --git a/Engine/Source/Editor/DataLayerEditor/Private/DataLayer/SDataLayerTreeLabel.h b/Engine/Source/Editor/DataLayerEditor/Private/DataLayer/SDataLayerTreeLabel.h
new file mode 100644
--- /dev/null
+++ b/Engine/Source/Editor/DataLayerEditor/Private/DataLayer/SDataLayerTreeLabel.h
@@ -0,0 +1,38 @@
+// Copyright Epic Games, Inc. All Rights Reserved.
+
+#pragma once
+
+#include "CoreMinimal.h"
+#include "Misc/Attribute.h"
+#include "SceneOutlinerFwd.h"
+#include "SceneOutlinerStandaloneTypes.h"
+#include "Widgets/SCompoundWidget.h"
+#include "Widgets/Views/STableRow.h"
+#include "Widgets/DeclarativeSyntaxSupport.h"
+
+struct FDataLayerTreeItem;
+class UDataLayer;
+
+struct SDataLayerTreeLabel : FSceneOutlinerCommonLabelData, public SCompoundWidget
+{
+	SLATE_BEGIN_ARGS(SDataLayerTreeLabel) {}
+	SLATE_END_ARGS()
+
+	void Construct(const FArguments& InArgs, FDataLayerTreeItem& DataLayerItem, ISceneOutliner& SceneOutliner, const STableRow<FSceneOutlinerTreeItemPtr>& InRow);
+
+private:
+
+	FText GetDisplayText() const;
+	FText GetTooltipText() const;
+	FText GetTypeText() const;
+	EVisibility GetTypeTextVisibility() const;
+	const FSlateBrush* GetIcon() const;
+	FText GetIconTooltip() const;
+	FSlateColor GetForegroundColor() const;
+	bool OnVerifyItemLabelChanged(const FText& InLabel, FText& OutErrorMessage);
+	void OnLabelCommitted(const FText& InLabel, ETextCommit::Type InCommitInfo);
+
+	TWeakPtr<FDataLayerTreeItem> TreeItemPtr;
+	TWeakObjectPtr<UDataLayer> DataLayerPtr;
+	TAttribute<FText> HighlightText;
+};
\ No newline at end of file
diff --git a/Engine/Source/Editor/DataLayerEditor/Private/DataLayerEditorModule.cpp b/Engine/Source/Editor/DataLayerEditor/Private/DataLayerEditorModule.cpp
new file mode 100644
--- /dev/null
+++ b/Engine/Source/Editor/DataLayerEditor/Private/DataLayerEditorModule.cpp
@@ -0,0 +1,32 @@
+// Copyright Epic Games, Inc. All Rights Reserved.
+
+#include "DataLayerEditorModule.h"
+#include "PropertyEditorModule.h"
+#include "Widgets/SWidget.h"
+#include "Widgets/DeclarativeSyntaxSupport.h"
+#include "Editor.h"
+#include "Modules/ModuleManager.h"
+#include "DataLayer/DataLayerPropertyTypeCustomization.h"
+#include "DataLayer/SDataLayerBrowser.h"
+#include "WorldPartition/DataLayer/DataLayer.h"
+
+IMPLEMENT_MODULE(FDataLayerEditorModule, DataLayerEditor );
+
+void FDataLayerEditorModule::StartupModule()
+{
+	FPropertyEditorModule& PropertyModule = FModuleManager::LoadModuleChecked<FPropertyEditorModule>("PropertyEditor");
+	PropertyModule.RegisterCustomPropertyTypeLayout(FActorDataLayer::StaticStruct()->GetFName(), FOnGetPropertyTypeCustomizationInstance::CreateLambda([] { return MakeShared<FDataLayerPropertyTypeCustomization>(); }));
+}
+
+void FDataLayerEditorModule::ShutdownModule()
+{
+	if (FPropertyEditorModule* PropertyModule = FModuleManager::GetModulePtr<FPropertyEditorModule>("PropertyEditor"))
+	{
+		PropertyModule->UnregisterCustomPropertyTypeLayout(FActorDataLayer::StaticStruct()->GetFName());
+	}
+}
+
+TSharedRef<SWidget> FDataLayerEditorModule::CreateDataLayerBrowser()
+{
+	return SNew(SDataLayerBrowser);
+}
\ No newline at end of file
diff --git a/Engine/Source/Editor/DataLayerEditor/Public/DataLayer/DataLayerAction.h b/Engine/Source/Editor/DataLayerEditor/Public/DataLayer/DataLayerAction.h
new file mode 100644
--- /dev/null
+++ b/Engine/Source/Editor/DataLayerEditor/Public/DataLayer/DataLayerAction.h
@@ -0,0 +1,12 @@
+// Copyright Epic Games, Inc. All Rights Reserved.
+
+#pragma once
+
+enum class EDataLayerAction : uint8
+{
+	Add,
+	Modify,
+	Delete,
+	Rename,
+	Reset
+};
\ No newline at end of file
diff --git a/Engine/Source/Editor/DataLayerEditor/Public/DataLayer/DataLayerEditorSubsystem.h b/Engine/Source/Editor/DataLayerEditor/Public/DataLayer/DataLayerEditorSubsystem.h
new file mode 100644
--- /dev/null
+++ b/Engine/Source/Editor/DataLayerEditor/Public/DataLayer/DataLayerEditorSubsystem.h
@@ -0,0 +1,606 @@
+// Copyright Epic Games, Inc. All Rights Reserved.
+
+#pragma once
+
+#include "CoreMinimal.h"
+#include "EditorSubsystem.h"
+#include "DataLayerAction.h"
+#include "DataLayerEditorSubsystem.generated.h"
+
+class AActor;
+class FDataLayersBroadcast;
+class FLevelEditorViewportClient;
+class UEditorEngine;
+class UDataLayer;
+class ULevel;
+class UWorld;
+template<typename TItemType> class IFilter;
+
+UCLASS()
+class DATALAYEREDITOR_API UDataLayerEditorSubsystem final : public UEditorSubsystem
+{
+	GENERATED_BODY()
+
+public:
+	static UDataLayerEditorSubsystem* Get();
+
+	typedef IFilter<const TWeakObjectPtr<AActor>&> FActorFilter;
+
+	/**
+	 *	Prepares for use
+	 */
+	virtual void Initialize(FSubsystemCollectionBase& Collection) override;
+
+	/**
+	 *	Internal cleanup
+	 */
+	virtual void Deinitialize() override;
+
+	/**
+	 *	Destructor
+	 */
+	virtual ~UDataLayerEditorSubsystem() {}
+
+	/* Broadcasts whenever one or more DataLayers are modified
+	 *
+	 * Actions
+	 * Add    : The specified ChangedDataLayer is a newly created UDataLayer
+	 * Modify : The specified ChangedDataLayer was just modified, if ChangedDataLayer is invalid then multiple DataLayers were modified.
+	 *          ChangedProperty specifies what field on the UDataLayer was changed, if NAME_None then multiple fields were changed
+	 * Delete : A DataLayer was deleted
+	 * Rename : The specified ChangedDataLayer was just renamed
+	 * Reset  : A large amount of changes have occurred to a number of DataLayers.
+	 */
+	DECLARE_EVENT_ThreeParams(UDataLayerEditorSubsystem, FOnDataLayerChanged, const EDataLayerAction /*Action*/, const TWeakObjectPtr<const UDataLayer>& /*ChangedDataLayer*/, const FName& /*ChangedProperty*/);
+	virtual FOnDataLayerChanged& OnDataLayerChanged() { return DataLayerChanged; }
+
+	/** Broadcasts whenever one or more Actors changed DataLayers*/
+	DECLARE_EVENT_OneParam(UDataLayerEditorSubsystem, FOnActorDataLayersChanged, const TWeakObjectPtr<AActor>& /*ChangedActor*/);
+	virtual FOnActorDataLayersChanged& OnActorDataLayersChanged() { return ActorDataLayersChanged; }
+	
+	///////////////////////////////////////////////////////////////////////////////////////////////////////////////////////
+	// Operations on an individual actor.
+
+	/**
+	 *	Checks to see if the specified actor is in an appropriate state to interact with DataLayers
+	 *
+	 *	@param	Actor	The actor to validate
+	 */
+	virtual bool IsActorValidForDataLayer(AActor* Actor);
+
+	/**
+	 *	Synchronizes an newly created Actor's DataLayers with the DataLayer system
+	 *
+	 *	@param	Actor	The actor to initialize
+	 */
+	virtual void InitializeNewActorDataLayers(AActor* Actor);
+	
+	/**
+	 * Adds the actor to the DataLayer.
+	 *
+	 * @param	Actor		The actor to add to the DataLayer
+	 * @param	DataLayer	The DataLayer to add the actor to
+	 * @return				true if the actor was added.  false is returned if the actor already belongs to the DataLayer.
+	 */
+	virtual bool AddActorToDataLayer(AActor* Actor, const UDataLayer* DataLayer);
+
+	/**
+	 * Adds the provided actor to the DataLayers.
+	 *
+	 * @param	Actor		The actor to add to the provided DataLayers
+	 * @param	DataLayers	A valid list of DataLayers.
+	 * @return				true if the actor was added to at least one of the provided DataLayers.
+	 */
+	virtual bool AddActorToDataLayers(AActor* Actor, const TArray<const UDataLayer*>& DataLayers);
+
+	/**
+	 * Removes an actor from the specified DataLayer.
+	 *
+	 * @param	Actor			The actor to remove from the provided DataLayer
+	 * @param	DataLayerToRemove	The DataLayer to remove the actor from
+	 * @return					true if the actor was removed from the DataLayer.  false is returned if the actor already belonged to the DataLayer.
+	 */
+	virtual bool RemoveActorFromDataLayer(AActor* Actor, const UDataLayer* DataLayerToRemove);
+
+	/**
+	 * Removes the provided actor from the DataLayers.
+	 *
+	 * @param	Actor		The actor to remove from the provided DataLayers
+	 * @param	DataLayers	A valid list of DataLayers.
+	 * @return				true if the actor was removed from at least one of the provided DataLayers.
+	 */
+	virtual bool RemoveActorFromDataLayers(AActor* Actor, const TArray<const UDataLayer*>& DataLayers);
+
+
+	/////////////////////////////////////////////////
+	// Operations on multiple actors.
+
+	/**
+	 * Add the actors to the DataLayer
+	 *
+	 * @param	Actors		The actors to add to the DataLayer
+	 * @param	DataLayer	The DataLayer to add to
+	 * @return				true if at least one actor was added to the DataLayer.  false is returned if all the actors already belonged to the DataLayer.
+	 */
+	virtual bool AddActorsToDataLayer(const TArray<AActor*>& Actors, const UDataLayer* DataLayer);
+
+	/**
+	 * Add the actors to the DataLayer
+	 *
+	 * @param	Actors		The actors to add to the DataLayer
+	 * @param	DataLayer	The DataLayer to add to
+	 * @return				true if at least one actor was added to the DataLayer.  false is returned if all the actors already belonged to the DataLayer.
+	 */
+	virtual bool AddActorsToDataLayer(const TArray<TWeakObjectPtr<AActor>>& Actors, const UDataLayer* DataLayer);
+
+	/**
+	 * Add the actors to the DataLayers
+	 *
+	 * @param	Actors		The actors to add to the DataLayers
+	 * @param	DataLayers	A valid list of DataLayers.
+	 * @return				true if at least one actor was added to at least one DataLayer.  false is returned if all the actors already belonged to all specified DataLayers.
+	 */
+	bool AddActorsToDataLayers(const TArray<AActor*>& Actors, const TArray<const UDataLayer*>& DataLayers);
+
+	/**
+	 * Add the actors to the DataLayers
+	 *
+	 * @param	Actors		The actors to add to the DataLayers
+	 * @param	DataLayers	A valid list of DataLayers.
+	 * @return				true if at least one actor was added to at least one DataLayer.  false is returned if all the actors already belonged to all specified DataLayers.
+	 */
+	virtual bool AddActorsToDataLayers(const TArray<TWeakObjectPtr<AActor>>& Actors, const TArray<const UDataLayer*>& DataLayers);
+
+	/**
+	 * Removes the actors from the specified DataLayer.
+	 *
+	 * @param	Actors			The actors to remove from the provided DataLayer
+	 * @param	DataLayerToRemove	The DataLayer to remove the actors from
+	 * @return					true if at least one actor was removed from the DataLayer.  false is returned if all the actors already belonged to the DataLayer.
+	 */
+	virtual bool RemoveActorsFromDataLayer(const TArray<AActor*>& Actors, const UDataLayer* DataLayer);
+
+	/**
+	 * Removes the actors from the specified DataLayer.
+	 *
+	 * @param	Actors			The actors to remove from the provided DataLayer
+	 * @param	DataLayerToRemove	The DataLayer to remove the actors from
+	 * @return					true if at least one actor was removed from the DataLayer.  false is returned if all the actors already belonged to the DataLayer.
+	 */
+	virtual bool RemoveActorsFromDataLayer(const TArray<TWeakObjectPtr<AActor>>& Actors, const UDataLayer* DataLayer);
+
+	/**
+	 * Remove the actors to the DataLayers
+	 *
+	 * @param	Actors		The actors to remove to the DataLayers
+	 * @param	DataLayers	A valid list of DataLayers.
+	 * @return				true if at least one actor was removed from at least one DataLayer.  false is returned if none of the actors belonged to any of the specified DataLayers.
+	 */
+	virtual bool RemoveActorsFromDataLayers(const TArray<AActor*>& Actors, const TArray<const UDataLayer*>& DataLayers);
+
+	/**
+	 * Remove the actors to the DataLayers
+	 *
+	 * @param	Actors		The actors to remove to the DataLayers
+	 * @param	DataLayers	A valid list of DataLayers.
+	 * @return				true if at least one actor was removed from at least one DataLayer.  false is returned if none of the actors belonged to any of the specified DataLayers.
+	 */
+	virtual bool RemoveActorsFromDataLayers(const TArray<TWeakObjectPtr<AActor>>& Actors, const TArray<const UDataLayer*>& DataLayers);
+
+	
+	/////////////////////////////////////////////////
+	// Operations on selected actors.
+
+	/**
+	 * Find and return the selected actors.
+	 *
+	 * @return				The selected AActor's as a TArray.
+	 */
+	virtual TArray<AActor*> GetSelectedActors() const;
+
+	/**
+	 * Adds selected actors to the DataLayer.
+	 *
+	 * @param	DataLayer	A DataLayer.
+	 * @return				true if at least one actor was added.  false is returned if all selected actors already belong to the DataLayer.
+	 */
+	virtual bool AddSelectedActorsToDataLayer(const UDataLayer* DataLayer);
+
+	/**
+	 * Adds selected actors to the DataLayers.
+	 *
+	 * @param	DataLayers	A valid list of DataLayers.
+	 * @return				true if at least one actor was added.  false is returned if all selected actors already belong to the DataLayers.
+	 */
+	virtual bool AddSelectedActorsToDataLayers(const TArray<const UDataLayer*>& DataLayers);
+
+	/**
+	 * Removes the selected actors from the DataLayer.
+	 *
+	 * @param	DataLayer	A DataLayer.
+	 * @return				true if at least one actor was added.  false is returned if all selected actors already belong to the DataLayer.
+	 */
+	virtual bool RemoveSelectedActorsFromDataLayer(const UDataLayer* DataLayer);
+
+	/**
+	 * Removes selected actors from the DataLayers.
+	 *
+	 * @param	DataLayers	A valid list of DataLayers.
+	 * @return				true if at least one actor was removed.
+	 */
+	virtual bool RemoveSelectedActorsFromDataLayers(const TArray<const UDataLayer*>& DataLayers);
+
+
+	/////////////////////////////////////////////////
+	// Operations on actors in DataLayers
+	
+	/**
+	 * Selects/de-selects actors belonging to the DataLayer.
+	 *
+	 * @param	DataLayer						A valid DataLayer.
+	 * @param	bSelect							If true actors are selected; if false, actors are deselected.
+	 * @param	bNotify							If true the Editor is notified of the selection change; if false, the Editor will not be notified.
+	 * @param	bSelectEvenIfHidden	[optional]	If true even hidden actors will be selected; if false, hidden actors won't be selected.
+	 * @return									true if at least one actor was selected/deselected.
+	 */
+	virtual bool SelectActorsInDataLayer(const UDataLayer* DataLayer, const bool bSelect, const bool bNotify, const bool bSelectEvenIfHidden = false);
+
+	/**
+	 * Selects/de-selects actors belonging to the DataLayer.
+	 *
+	 * @param	DataLayer						A valid DataLayer.
+	 * @param	bSelect							If true actors are selected; if false, actors are deselected.
+	 * @param	bNotify							If true the Editor is notified of the selection change; if false, the Editor will not be notified.
+	 * @param	bSelectEvenIfHidden	[optional]	If true even hidden actors will be selected; if false, hidden actors won't be selected.
+	 * @param	Filter	[optional]				Actor that don't pass the specified filter restrictions won't be selected.
+	 * @return									true if at least one actor was selected/deselected.
+	 */
+	virtual bool SelectActorsInDataLayer(const UDataLayer* DataLayer, const bool bSelect, const bool bNotify, const bool bSelectEvenIfHidden, const TSharedPtr<FActorFilter>& Filter);
+	/**
+	 * Selects/de-selects actors belonging to the DataLayers.
+	 *
+	 * @param	DataLayers						A valid list of DataLayers.
+	 * @param	bSelect							If true actors are selected; if false, actors are deselected.
+	 * @param	bNotify							If true the Editor is notified of the selection change; if false, the Editor will not be notified
+	 * @param	bSelectEvenIfHidden	[optional]	If true even hidden actors will be selected; if false, hidden actors won't be selected.
+	 * @return									true if at least one actor was selected/deselected.
+	 */
+	virtual bool SelectActorsInDataLayers(const TArray<const UDataLayer*>& DataLayers, const bool bSelect, const bool bNotify, const bool bSelectEvenIfHidden = false);
+
+	/**
+	 * Selects/de-selects actors belonging to the DataLayers.
+	 *
+	 * @param	DataLayers						A valid list of DataLayers.
+	 * @param	bSelect							If true actors are selected; if false, actors are deselected.
+	 * @param	bNotify							If true the Editor is notified of the selection change; if false, the Editor will not be notified
+	 * @param	bSelectEvenIfHidden	[optional]	If true even hidden actors will be selected; if false, hidden actors won't be selected.
+	 * @param	Filter	[optional]				Actor that don't pass the specified filter restrictions won't be selected.
+	 * @return									true if at least one actor was selected/deselected.
+	 */
+	virtual bool SelectActorsInDataLayers(const TArray<const UDataLayer*>& DataLayers, const bool bSelect, const bool bNotify, const bool bSelectEvenIfHidden, const TSharedPtr<FActorFilter>& Filter);
+
+
+	/////////////////////////////////////////////////
+	// Operations on actor viewport visibility regarding DataLayers
+
+	/**
+	 * Updates the visibility for all actors for all views.
+	 *
+	 * @param DataLayerThatChanged  If one DataLayer was changed (toggled in view pop-up, etc), then we only need to modify actors that use that DataLayer.
+	 */
+	virtual void UpdateAllViewVisibility(const UDataLayer* DataLayerThatChanged);
+
+	/**
+	 * Updates the per-view visibility for all actors for the given view
+	 *
+	 * @param ViewportClient				The viewport client to update visibility on
+	 * @param DataLayerThatChanged [optional]	If one DataLayer was changed (toggled in view pop-up, etc), then we only need to modify actors that use that DataLayer
+	 */
+	virtual void UpdatePerViewVisibility(FLevelEditorViewportClient* ViewportClient, const UDataLayer* DataLayerThatChanged = nullptr);
+
+	/**
+	 * Updates per-view visibility for the given actor in the given view
+	 *
+	 * @param ViewportClient				The viewport client to update visibility on
+	 * @param Actor								Actor to update
+	 * @param bReregisterIfDirty [optional]		If true, the actor will reregister itself to give the rendering thread updated information
+	 */
+	virtual void UpdateActorViewVisibility(FLevelEditorViewportClient* ViewportClient, AActor* Actor, const bool bReregisterIfDirty = true);
+
+	/**
+	 * Updates per-view visibility for the given actor for all views
+	 *
+	 * @param Actor		Actor to update
+	 */
+	virtual void UpdateActorAllViewsVisibility(AActor* Actor);
+
+	/**
+	 * Removes the corresponding visibility bit from all actors (slides the later bits down 1)
+	 *
+	 * @param ViewportClient	The viewport client to update visibility on
+	 */
+	virtual void RemoveViewFromActorViewVisibility(FLevelEditorViewportClient* ViewportClient);
+
+	/**
+	 * Updates the provided actors visibility in the viewports
+	 *
+	 * @param	Actor						Actor to update
+	 * @param	bOutSelectionChanged [OUT]	Whether the Editors selection changed
+	 * @param	bOutActorModified [OUT]		Whether the actor was modified
+	 * @param	bNotifySelectionChange		If true the Editor is notified of the selection change; if false, the Editor will not be notified
+	 * @param	bRedrawViewports			If true the viewports will be redrawn; if false, they will not
+	 */
+	virtual bool UpdateActorVisibility(AActor* Actor, bool& bOutSelectionChanged, bool& bOutActorModified, const bool bNotifySelectionChange, const bool bRedrawViewports);
+
+	/**
+	 * Updates the visibility of all actors in the viewports
+	 *
+	 * @param	bNotifySelectionChange		If true the Editor is notified of the selection change; if false, the Editor will not be notified
+	 * @param	bRedrawViewports			If true the viewports will be redrawn; if false, they will not
+	 */
+	virtual bool UpdateAllActorsVisibility(const bool bNotifySelectionChange, const bool bRedrawViewports);
+
+
+	/////////////////////////////////////////////////
+	// Operations on DataLayers
+
+	/**
+	 *	Appends all the actors associated with the specified DataLayer.
+	 *
+	 *	@param	DataLayer			The DataLayer to find actors for.
+	 *	@param	InOutActors			The list to append the found actors to.
+	 */
+	virtual void AppendActorsFromDataLayer(const UDataLayer* DataLayer, TArray<AActor*>& InOutActors) const;
+
+	/**
+	 *	Appends all the actors associated with the specified DataLayer.
+	 *
+	 *	@param	DataLayer			The DataLayer to find actors for.
+	 *	@param	InOutActors			The list to append the found actors to.
+	 *  @param	Filter	[optional]	Actor that don't pass the specified filter restrictions won't be selected.
+	 */
+	virtual void AppendActorsFromDataLayer(const UDataLayer* DataLayer, TArray<AActor*>& InOutActors, const TSharedPtr<FActorFilter>& Filter) const;
+	
+	/**
+	 *	Appends all the actors associated with the specified DataLayer.
+	 *
+	 *	@param	DataLayer			The DataLayer to find actors for.
+	 *	@param	InOutActors			The list to append the found actors to.
+	 *  @param	Filter	[optional]	Actor that don't pass the specified filter restrictions won't be selected.
+	 */
+	virtual void AppendActorsFromDataLayer(const UDataLayer* DataLayer, TArray<TWeakObjectPtr<AActor>>& InOutActors, const TSharedPtr<FActorFilter>& Filter = TSharedPtr<FActorFilter>(nullptr)) const;
+	
+	/**
+	 *	Appends all the actors associated with ANY of the specified DataLayers.
+	 *
+	 *	@param	DataLayers			The DataLayers to find actors for.
+	 *	@param	InOutActors			The list to append the found actors to.
+	 */
+	virtual void AppendActorsFromDataLayers(const TArray<const UDataLayer*>& DataLayers, TArray<AActor*>& InOutActors) const;
+
+	/**
+	 *	Appends all the actors associated with ANY of the specified DataLayers.
+	 *
+	 *	@param	DataLayers			The DataLayers to find actors for.
+	 *	@param	InOutActors			The list to append the found actors to.
+	 *  @param	Filter	[optional]	Actor that don't pass the specified filter restrictions won't be selected.
+	 */
+	virtual void AppendActorsFromDataLayers(const TArray<const UDataLayer*>& DataLayers, TArray<AActor*>& InOutActors, const TSharedPtr<FActorFilter>& Filter) const;
+
+	/**
+	 *	Appends all the actors associated with ANY of the specified DataLayers.
+	 *
+	 *	@param	DataLayers			The DataLayers to find actors for.
+	 *	@param	InOutActors			The list to append the found actors to.
+	 *  @param	Filter	[optional]	Actor that don't pass the specified filter restrictions won't be selected.
+	 */
+	virtual void AppendActorsFromDataLayers(const TArray<const UDataLayer*>& DataLayers, TArray<TWeakObjectPtr<AActor>>& InOutActors, const TSharedPtr<FActorFilter>& Filter = TSharedPtr<FActorFilter>(nullptr)) const;
+
+	/**
+	 *	Gets all the actors associated with the specified DataLayer. Analog to AppendActorsFromDataLayer but it returns rather than appends the actors.
+	 *
+	 *	@param	DataLayer			The DataLayer to find actors for.
+	 *	@return						The list to assign the found actors to.
+	 */
+	virtual TArray<AActor*> GetActorsFromDataLayer(const UDataLayer* DataLayer) const;
+
+	/**
+	 *	Gets all the actors associated with the specified DataLayer. Analog to AppendActorsFromDataLayer but it returns rather than appends the actors.
+	 *
+	 *	@param	DataLayer			The DataLayer to find actors for.
+	 *  @param	Filter	[optional]	Actor that don't pass the specified filter restrictions won't be selected.
+	 *	@return						The list to assign the found actors to.
+	 */
+	virtual TArray<AActor*> GetActorsFromDataLayer(const UDataLayer* DataLayer, const TSharedPtr<FActorFilter>& Filter) const;
+
+	/**
+	 *	Gets all the actors associated with ANY of the specified DataLayers. Analog to AppendActorsFromDataLayers but it returns rather than appends the actors.
+	 *
+	 *	@param	DataLayers			The DataLayers to find actors for.
+	 *	@return						The list to assign the found actors to.
+	 */
+	virtual TArray<AActor*> GetActorsFromDataLayers(const TArray<const UDataLayer*>& DataLayers) const;
+
+	/**
+	 *	Gets all the actors associated with ANY of the specified DataLayers. Analog to AppendActorsFromDataLayers but it returns rather than appends the actors.
+	 *
+	 *	@param	DataLayers			The DataLayers to find actors for.
+	 *  @param	Filter	[optional]	Actor that don't pass the specified filter restrictions won't be selected.
+	 *	@return						The list to assign the found actors to.
+	 */
+	virtual TArray<AActor*> GetActorsFromDataLayers(const TArray<const UDataLayer*>& DataLayers, const TSharedPtr<FActorFilter>& Filter) const;
+
+	/**
+	 * Changes the DataLayer's visibility to the provided state
+	 *
+	 * @param	DataLayer	The DataLayer to affect.
+	 * @param	bIsVisible	If true the DataLayer will be visible; if false, the DataLayer will not be visible.
+	 */
+	virtual void SetDataLayerVisibility(UDataLayer* DataLayer, const bool bIsVisible);
+
+	/**
+	 * Changes visibility of the DataLayers to the provided state
+	 *
+	 * @param	DataLayers	The DataLayers to affect
+	 * @param	bIsVisible	If true the DataLayers will be visible; if false, the DataLayers will not be visible
+	 */
+	virtual void SetDataLayersVisibility(const TArray<UDataLayer*>& DataLayers, const bool bIsVisible);
+
+	/**
+	 * Toggles the DataLayer's visibility
+	 *
+	 * @param DataLayer	The DataLayer to affect
+	 */
+	virtual void ToggleDataLayerVisibility(UDataLayer* DataLayer);
+
+	/**
+	 * Toggles the visibility of all of the DataLayers
+	 *
+	 * @param	DataLayers	The DataLayers to affect
+	 */
+	virtual void ToggleDataLayersVisibility(const TArray<UDataLayer*>& DataLayers);
+
+	/**
+	 * Changes the DataLayer's IsDynamicallyLoaded flag to the provided state
+	 *
+	 * @param	DataLayer				The DataLayer to affect.
+	 * @param	bIsDynamicallyLoaded	If true the DataLayer will affect runtime streaming and actors will only be loaded if the DataLayer is active.
+	 */
+	virtual void SetDataLayerIsDynamicallyLoaded(UDataLayer* DataLayer, const bool bIsDynamicallyLoaded);
+
+	/**
+	 * Changes the IsDynamicallyLoaded flag of the DataLayers to the provided state
+	 *
+	 * @param	DataLayers				The DataLayers to affect
+	 * @param	bIsDynamicallyLoaded	If true the DataLayer will affect runtime streaming and actors will only be loaded if the DataLayer is active.
+	 */
+	virtual void SetDataLayersIsDynamicallyLoaded(const TArray<UDataLayer*>& DataLayers, const bool bIsDynamicallyLoaded);
+
+	/**
+	 * Toggles the DataLayer's IsDynamicallyLoaded flag
+	 *
+	 * @param DataLayer	The DataLayer to affect
+	 */
+	virtual void ToggleDataLayerIsDynamicallyLoaded(UDataLayer* DataLayer);
+
+	/**
+	 * Toggles the IsDynamicallyLoaded flag of all of the DataLayers
+	 *
+	 * @param	DataLayers	The DataLayers to affect
+	 */
+	virtual void ToggleDataLayersIsDynamicallyLoaded(const TArray<UDataLayer*>& DataLayers);
+
+	/**
+	 * Set the visibility of all DataLayers to true
+	 */
+	virtual void MakeAllDataLayersVisible();
+
+	/**
+	 * Gets the UDataLayer Object of the DataLayer label
+	 *
+	 * @param	DataLayerLabel	The label of the DataLayer whose UDataLayer Object is returned
+	 * @return					The UDataLayer Object of the provided DataLayer label
+	 */
+	virtual const UDataLayer* GetDataLayerFromLabel(const FName& DataLayerLabel) const;
+
+	/**
+	 * Gets the UDataLayer Object of the DataLayer name
+	 *
+	 * @param	DataLayerName	The name of the DataLayer whose UDataLayer Object is returned
+	 * @return					The UDataLayer Object of the provided DataLayer name
+	 */
+	virtual const UDataLayer* GetDataLayerFromName(const FName& DataLayerName) const;
+	
+	/**
+	 * Attempts to get the UDataLayer Object of the provided DataLayer label.
+	 *
+	 * @param	DataLayerLabel		The label of the DataLayer whose UDataLayer Object to retrieve
+	 * @param	OutDataLayer[OUT] 	Set to the UDataLayer Object of the DataLayer label. Set to Invalid if no UDataLayer Object exists.
+	 * @return					If true a valid UDataLayer Object was found and set to OutDataLayer; if false, a valid UDataLayer object was not found and invalid set to OutDataLayer
+	 */
+	virtual bool TryGetDataLayerFromLabel(const FName& DataLayerLabel, const UDataLayer*& OutDataLayer);
+
+	/**
+	 * Gets all known DataLayers and appends them to the provided array
+	 *
+	 * @param OutDataLayers[OUT] Output array to store all known DataLayers
+	 */
+	virtual void AddAllDataLayersTo(TArray<TWeakObjectPtr<UDataLayer>>& OutDataLayers) const;
+
+	/**
+	 * Creates a UDataLayer Object
+	 *
+	 * @return	The newly created UDataLayer Object
+	 */
+	 UDataLayer* CreateDataLayer();
+
+	/**
+	 * Gets the WorldDataLayers for this world
+	 *
+	 * @param	bCreateIfNotFound	If true, created the AWorldDataLayers object if not found.
+	 * @return	The AWorldDataLayers Object
+	 */
+	class AWorldDataLayers* GetWorldDataLayers(bool bCreateIfNotFound = false);
+
+	/**
+	 * Gets the WorldDataLayers for this world
+	 *
+	 * @return	The AWorldDataLayers Object
+	 */
+	const class AWorldDataLayers* GetWorldDataLayers() const;
+
+	/**
+	 * Deletes all of the provided DataLayers, disassociating all actors from them
+	 *
+	 * @param DataLayersToDelete	A valid list of DataLayer.
+	 */
+	virtual void DeleteDataLayers(const TArray<const UDataLayer*>& DataLayersToDelete);
+
+	/**
+	 * Deletes the provided DataLayer, disassociating all actors from them
+	 *
+	 * @param DataLayerToDelete		A valid DataLayer
+	 */
+	virtual void DeleteDataLayer(const UDataLayer* DataLayerToDelete);
+
+	/**
+	 * Renames the provided DataLayer to the provided new label
+	 *
+	 * @param	DataLayer			The DataLayer to be renamed
+	 * @param	NewDataLayerLabel	The new label for the DataLayer to be renamed
+	 */
+	virtual bool RenameDataLayer(UDataLayer* DataLayer, const FName& NewDataLayerLabel);
+
+	/**
+	 * Get the current UWorld object.
+	 *
+	 * @return						The UWorld* object
+	 */
+	UWorld* GetWorld() const; // Fallback to GWorld
+
+	/**
+	 * Delegate handler for FEditorDelegates::MapChange. It internally calls DataLayerChanged.Broadcast.
+	 **/
+	void EditorMapChange();
+	
+	/**
+	 * Delegate handler for FEditorDelegates::RefreshDataLayerBrowser. It internally calls UpdateAllActorsVisibility to refresh the actors of each DataLayer.
+	 **/
+	void EditorRefreshDataLayerBrowser();
+
+	/**
+	 * Delegate handler for FEditorDelegates::PostUndoRedo. It internally calls DataLayerChanged.Broadcast and UpdateAllActorsVisibility to refresh the actors of each DataLayer.
+	 **/
+	void PostUndoRedo();
+
+private:
+
+	/** Fires whenever one or more DataLayer changes */
+	FOnDataLayerChanged DataLayerChanged;
+
+	/**	Fires whenever one or more actor DataLayer changes */
+	FOnActorDataLayersChanged ActorDataLayersChanged;
+
+	/** Auxiliary class that sets the callback functions for multiple delegates */
+	TSharedPtr<class FDataLayersBroadcast> DataLayersBroadcast;
+};
diff --git a/Engine/Source/Editor/DataLayerEditor/Public/DataLayerEditorModule.h b/Engine/Source/Editor/DataLayerEditor/Public/DataLayerEditorModule.h
new file mode 100644
--- /dev/null
+++ b/Engine/Source/Editor/DataLayerEditor/Public/DataLayerEditorModule.h
@@ -0,0 +1,44 @@
+// Copyright Epic Games, Inc. All Rights Reserved.
+
+#pragma once
+
+#include "CoreMinimal.h"
+#include "Modules/ModuleInterface.h"
+#include "Framework/Commands/UICommandList.h"
+#include "Framework/MultiBox/MultiBoxExtender.h"
+
+class AActor;
+
+/**
+ * The module holding all of the UI related pieces for DataLayers
+ */
+class FDataLayerEditorModule : public IModuleInterface
+{
+public:
+
+	/**
+	 * Called right after the module DLL has been loaded and the module object has been created
+	 */
+	virtual void StartupModule();
+
+	/**
+	 * Called before the module is unloaded, right before the module object is destroyed.
+	 */
+	virtual void ShutdownModule();
+
+	/**
+	 * Creates a DataLayer Browser widget
+	 */
+	virtual TSharedRef<class SWidget> CreateDataLayerBrowser();
+
+	/** Delegates to be called to extend the DataLayers menus */
+	DECLARE_DELEGATE_RetVal_OneParam(TSharedRef<FExtender>, FDataLayersMenuExtender, const TSharedRef<FUICommandList>);
+	virtual TArray<FDataLayersMenuExtender>& GetAllDataLayersMenuExtenders() {return DataLayersMenuExtenders;}
+
+private:
+
+	/** All extender delegates for the DataLayers menus */
+	TArray<FDataLayersMenuExtender> DataLayersMenuExtenders;
+};
+
+
diff --git a/Engine/Source/Editor/EditorStyle/Private/StarshipStyle.cpp b/Engine/Source/Editor/EditorStyle/Private/StarshipStyle.cpp
--- a/Engine/Source/Editor/EditorStyle/Private/StarshipStyle.cpp
+++ b/Engine/Source/Editor/EditorStyle/Private/StarshipStyle.cpp
@@ -1158,7 +1158,21 @@ void FStarshipEditorStyle::FStyle::SetupGeneralStyles()
 		Set( "LayerBrowserButton.LabelFont", DEFAULT_FONT( "Regular", 8 ) );
 	}
 
-	
+	// DataLayer
+	{
+		Set("DataLayer.Icon16x", new IMAGE_BRUSH("Icons/datalayer_16x", Icon16x16));
+		Set("DataLayer.DynamicallyLoaded", new IMAGE_BRUSH("/Icons/datalayer_dynamicallyloaded_16x", Icon16x16));
+		Set("DataLayer.NotDynamicallyLoaded", new IMAGE_BRUSH("/Icons/datalayer_not_dynamicallyloaded_16x", Icon16x16));
+		Set("DataLayerBrowser.DataLayerContentsQuickbarBackground", new CORE_BOX_BRUSH("Common/DarkGroupBorder", 4.f / 16.f));
+		Set("DataLayerBrowser.ExploreDataLayerContents", new IMAGE_BRUSH("Icons/datalayer_explore_contents", Icon16x16));
+		Set("DataLayerBrowser.ReturnToDataLayersList", new IMAGE_BRUSH("Icons/datalayer_exit_contents", Icon16x16));
+		Set("DataLayerBrowser.Actor.RemoveFromDataLayer", new IMAGE_BRUSH("Icons/Cross_12x", Icon12x12));
+		Set("DataLayerBrowserButton", FButtonStyle(Button)
+			.SetNormal(FSlateNoResource())
+			.SetHovered(BOX_BRUSH("Common/RoundedSelection_16x", 4.0f / 16.0f, SelectionColor))
+			.SetPressed(BOX_BRUSH("Common/RoundedSelection_16x", 4.0f / 16.0f, SelectionColor_Pressed))
+		);
+	}
 
 	// Scene Outliner
 	{
@@ -4303,6 +4317,7 @@ void FStarshipEditorStyle::FStyle::SetupLevelEditorStyle()
 			Set( "LevelEditor.Tabs.WorldPartition", new IMAGE_BRUSH( "/Icons/icon_levels_partitionbutton_16x", Icon16x16 ) );
 
 			Set( "LevelEditor.Tabs.Layers",                 new IMAGE_BRUSH_SVG("Starship/Common/Layers", Icon16x16) );
+			Set( "LevelEditor.Tabs.DataLayers",				new IMAGE_BRUSH_SVG("Starship/Common/DataLayers", Icon16x16));
 			Set( "LevelEditor.Tabs.ComposureCompositing",   new IMAGE_BRUSH_SVG("Starship/Common/ComposureCompositing", Icon16x16) );
 			Set( "LevelEditor.Tabs.USDStage",   			new IMAGE_BRUSH_SVG("Starship/Common/USDStage", Icon16x16) );
 
diff --git a/Engine/Source/Editor/Layers/Private/LayerViewModel.cpp b/Engine/Source/Editor/Layers/Private/LayerViewModel.cpp
--- a/Engine/Source/Editor/Layers/Private/LayerViewModel.cpp
+++ b/Engine/Source/Editor/Layers/Private/LayerViewModel.cpp
@@ -103,16 +103,6 @@ void FLayerViewModel::ToggleVisibility()
 	VisibilityToggledEvent.Broadcast( AsShared() );
 }
 
-void FLayerViewModel::ToggleActorsLoading()
-{
-	if (!Layer.IsValid())
-	{
-		return;
-	}
-
-	ActorsLoadingToggledEvent.Broadcast(AsShared());
-}
-
 
 bool FLayerViewModel::CanRenameTo( const FName& NewLayerName, FString& OutMessage ) const
 {
diff --git a/Engine/Source/Editor/Layers/Private/LayerViewModel.h b/Engine/Source/Editor/Layers/Private/LayerViewModel.h
--- a/Engine/Source/Editor/Layers/Private/LayerViewModel.h
+++ b/Engine/Source/Editor/Layers/Private/LayerViewModel.h
@@ -56,11 +56,6 @@ public:
 	 */
 	void ToggleVisibility();
 
-	/** 
-	 *	Toggles the specified layer's actors loading behavior
-	 */
-	void ToggleActorsLoading();
-
 	/**
 	 *	Returns whether the layer can be assigned the specified name
 	 *
@@ -239,7 +234,4 @@ private:
 
 	/** Broadcasts whenever a layer visibility toggle is requested */
 	FVisibilityToggledEvent VisibilityToggledEvent;
-
-	/** Broadcasts whenever a layer actors loading toggle is requested */
-	FVisibilityToggledEvent ActorsLoadingToggledEvent;
 };
diff --git a/Engine/Source/Editor/Layers/Private/SLayersViewRow.h b/Engine/Source/Editor/Layers/Private/SLayersViewRow.h
--- a/Engine/Source/Editor/Layers/Private/SLayersViewRow.h
+++ b/Engine/Source/Editor/Layers/Private/SLayersViewRow.h
@@ -111,17 +111,6 @@ private:
 		return FReply::Handled();
 	}
 
-	/**
-	 *	Called when the user clicks on the actors loading icon for a layer's row widget
-	 *
-	 *	@return	A reply that indicated whether this event was handled.
-	 */
-	FReply OnToggleActorsLoading()
-	{
-		ViewModel->ToggleActorsLoading();
-		return FReply::Handled();
-	}
-
 	/**
 	 *	Called to get the Slate Image Brush representing the visibility state of
 	 *	the layer this row widget represents
diff --git a/Engine/Source/Editor/LevelEditor/LevelEditor.Build.cs b/Engine/Source/Editor/LevelEditor/LevelEditor.Build.cs
--- a/Engine/Source/Editor/LevelEditor/LevelEditor.Build.cs
+++ b/Engine/Source/Editor/LevelEditor/LevelEditor.Build.cs
@@ -17,7 +17,8 @@ public class LevelEditor : ModuleRules
                 "AppFramework",
                 "PortalServices",
                 "Persona",
-            }
+				"DataLayerEditor",
+			}
 		);
 
 		PublicIncludePathModuleNames.AddRange(
@@ -78,6 +79,7 @@ public class LevelEditor : ModuleRules
 				"EditorSubsystem",
 				"EnvironmentLightingViewer",
 				"DesktopPlatform",
+				"DataLayerEditor",
 			}
 		);
 
@@ -91,7 +93,7 @@ public class LevelEditor : ModuleRules
 				"SlateReflector",
 				"AutomationWindow",
 				"Layers",
-                "WorldBrowser",
+				"WorldBrowser",
 				"WorldPartitionEditor",
 				"EditorWidgets",
 				"AssetTools",
diff --git a/Engine/Source/Editor/LevelEditor/Private/LevelEditor.cpp b/Engine/Source/Editor/LevelEditor/Private/LevelEditor.cpp
--- a/Engine/Source/Editor/LevelEditor/Private/LevelEditor.cpp
+++ b/Engine/Source/Editor/LevelEditor/Private/LevelEditor.cpp
@@ -65,6 +65,7 @@ const FName LevelEditorTabIds::LevelEditorBuildAndSubmit(TEXT("LevelEditorBuildA
 const FName LevelEditorTabIds::LevelEditorSceneOutliner(TEXT("LevelEditorSceneOutliner"));
 const FName LevelEditorTabIds::LevelEditorStatsViewer(TEXT("LevelEditorStatsViewer"));
 const FName LevelEditorTabIds::LevelEditorLayerBrowser(TEXT("LevelEditorLayerBrowser"));
+const FName LevelEditorTabIds::LevelEditorDataLayerBrowser(TEXT("LevelEditorDataLayerBrowser"));
 const FName LevelEditorTabIds::Sequencer(TEXT("Sequencer"));
 const FName LevelEditorTabIds::SequencerGraphEditor(TEXT("SequencerGraphEditor"));
 const FName LevelEditorTabIds::WorldSettings(TEXT("WorldSettingsTab"));
diff --git a/Engine/Source/Editor/LevelEditor/Private/LevelViewportActions.cpp b/Engine/Source/Editor/LevelEditor/Private/LevelViewportActions.cpp
--- a/Engine/Source/Editor/LevelEditor/Private/LevelViewportActions.cpp
+++ b/Engine/Source/Editor/LevelEditor/Private/LevelViewportActions.cpp
@@ -78,6 +78,12 @@ void FLevelViewportCommands::RegisterCommands()
 		UI_COMMAND( HideAllLayers, "Hide All Layers", "Hides all layers", EUserInterfaceActionType::Button, FInputChord() );
 	}
 
+	// Generate a command for show/hide all DataLayers
+	{
+		UI_COMMAND(ShowAllDataLayers, "Show All Data Layers", "Shows all Data Layers", EUserInterfaceActionType::Button, FInputChord());
+		UI_COMMAND(HideAllDataLayers, "Hide All Data Layers", "Hides all Data Layers", EUserInterfaceActionType::Button, FInputChord());
+	}
+
 	// Generate a command for each sprite category
 	{
 		UI_COMMAND( ShowAllSprites, "Show All Sprites", "Shows all sprites", EUserInterfaceActionType::Button, FInputChord() );
diff --git a/Engine/Source/Editor/LevelEditor/Private/SLevelEditor.cpp b/Engine/Source/Editor/LevelEditor/Private/SLevelEditor.cpp
--- a/Engine/Source/Editor/LevelEditor/Private/SLevelEditor.cpp
+++ b/Engine/Source/Editor/LevelEditor/Private/SLevelEditor.cpp
@@ -36,6 +36,8 @@
 #include "Editor/Layers/Public/LayersModule.h"
 #include "Editor/WorldBrowser/Public/WorldBrowserModule.h"
 #include "Editor/WorldPartitionEditor/Public/WorldPartitionEditorModule.h"
+#include "Editor/DataLayerEditor/Public/DataLayerEditorModule.h"
+#include "WorldPartition/WorldPartitionSubsystem.h"
 #include "Toolkits/ToolkitManager.h"
 #include "PropertyEditorModule.h"
 #include "Interfaces/IMainFrameModule.h"
@@ -854,6 +856,22 @@ TSharedRef<SDockTab> SLevelEditor::SpawnLevelEditorTab( const FSpawnTabArgs& Arg
 				WorldPartitionEditorModule.CreateWorldPartitionEditor()
 			];
 	}
+	else if (TabIdentifier == LevelEditorTabIds::LevelEditorDataLayerBrowser)
+	{
+		FDataLayerEditorModule& DataLayerEditorModule = FModuleManager::LoadModuleChecked<FDataLayerEditorModule>( "DataLayerEditor" );
+		return SNew(SDockTab)
+			.Icon(FEditorStyle::GetBrush("LevelEditor.Tabs.DataLayers"))
+			.Label(NSLOCTEXT("LevelEditor", "DataLayersTabTitle", "Data Layers"))
+			[
+				SNew(SBorder)
+				.Padding(0)
+				.BorderImage(FEditorStyle::GetBrush("ToolPanel.GroupBorder"))
+				.AddMetaData<FTutorialMetaData>(FTutorialMetaData(TEXT("DataLayerBrowser"), TEXT("LevelEditorDataLayerBrowser")))
+				[
+					DataLayerEditorModule.CreateDataLayerBrowser()
+				]
+			];
+	}
 	else if( TabIdentifier == TEXT("Sequencer") )
 	{
 		if (FSlateStyleRegistry::FindSlateStyle("LevelSequenceEditorStyle"))
@@ -945,6 +963,16 @@ bool SLevelEditor::HasAnyHostedEditorModeToolkit() const
 	return false;
 }
 
+bool SLevelEditor::CanSpawnLayerBrowser(const FSpawnTabArgs& Args) const
+{
+	return !UWorld::HasSubsystem<UWorldPartitionSubsystem>(GetWorld());
+}
+
+bool SLevelEditor::CanSpawnDataLayerBrowser(const FSpawnTabArgs& Args) const
+{
+	return UWorld::HasSubsystem<UWorldPartitionSubsystem>(GetWorld());
+}
+
 TSharedPtr<SDockTab> SLevelEditor::TryInvokeTab( FName TabID )
 {
 	TSharedPtr<FTabManager> LevelEditorTabManager = GetTabManager();
@@ -1237,13 +1265,24 @@ TSharedRef<SWidget> SLevelEditor::RestoreContentArea( const TSharedRef<SDockTab>
 
 		{
 			const FSlateIcon LayersIcon(FEditorStyle::GetStyleSetName(), "LevelEditor.Tabs.Layers");
-			LevelEditorTabManager->RegisterTabSpawner(LevelEditorTabIds::LevelEditorLayerBrowser, FOnSpawnTab::CreateSP<SLevelEditor, FName, FString>(this, &SLevelEditor::SpawnLevelEditorTab, LevelEditorTabIds::LevelEditorLayerBrowser, FString()) )
+			LevelEditorTabManager->RegisterTabSpawner(LevelEditorTabIds::LevelEditorLayerBrowser, FOnSpawnTab::CreateSP<SLevelEditor, FName, FString>(this, &SLevelEditor::SpawnLevelEditorTab, LevelEditorTabIds::LevelEditorLayerBrowser, FString()),
+																								  FCanSpawnTab::CreateSP(this, &SLevelEditor::CanSpawnLayerBrowser))
 				.SetDisplayName(NSLOCTEXT("LevelEditorTabs", "LevelEditorLayerBrowser", "Layers"))
 				.SetTooltipText(NSLOCTEXT("LevelEditorTabs", "LevelEditorLayerBrowserTooltipText", "Open the Layers tab. Use this to manage which actors in the world belong to which layers."))
 				.SetGroup( MenuStructure.GetLevelEditorCategory() )
 				.SetIcon( LayersIcon );
 		}
 
+		{
+			const FSlateIcon DataLayersIcon(FEditorStyle::GetStyleSetName(), "LevelEditor.Tabs.DataLayers");
+			LevelEditorTabManager->RegisterTabSpawner(LevelEditorTabIds::LevelEditorDataLayerBrowser, FOnSpawnTab::CreateSP<SLevelEditor, FName, FString>(this, &SLevelEditor::SpawnLevelEditorTab, LevelEditorTabIds::LevelEditorDataLayerBrowser, FString()),
+																									  FCanSpawnTab::CreateSP(this, &SLevelEditor::CanSpawnDataLayerBrowser))
+				.SetDisplayName(NSLOCTEXT("LevelEditorTabs", "LevelEditorDataLayerBrowser", "Data Layers"))
+				.SetTooltipText(NSLOCTEXT("LevelEditorTabs", "LevelEditorDataLayerBrowserTooltipText", "Open the Data Layers tab. Use this to manage which actors in the world belong to which Data Layers."))
+				.SetGroup(MenuStructure.GetLevelEditorCategory())
+				.SetIcon(DataLayersIcon);
+		}
+
 		{
 			const FSlateIcon LayersIcon(FEditorStyle::GetStyleSetName(), "LevelEditor.Tabs.HLOD");
 			LevelEditorTabManager->RegisterTabSpawner(LevelEditorTabIds::LevelEditorHierarchicalLODOutliner, FOnSpawnTab::CreateSP<SLevelEditor, FName, FString>(this, &SLevelEditor::SpawnLevelEditorTab, LevelEditorTabIds::LevelEditorHierarchicalLODOutliner, FString()))
diff --git a/Engine/Source/Editor/LevelEditor/Private/SLevelEditor.h b/Engine/Source/Editor/LevelEditor/Private/SLevelEditor.h
--- a/Engine/Source/Editor/LevelEditor/Private/SLevelEditor.h
+++ b/Engine/Source/Editor/LevelEditor/Private/SLevelEditor.h
@@ -146,6 +146,8 @@ private:
 	TSharedRef<SDockTab> SpawnLevelEditorTab(const FSpawnTabArgs& Args, FName TabIdentifier, FString InitializationPayload);
 	bool CanSpawnEditorModeToolbarTab(const FSpawnTabArgs& Args) const;
 	bool CanSpawnEditorModeToolboxTab(const FSpawnTabArgs& Args) const;
+	bool CanSpawnLayerBrowser(const FSpawnTabArgs& Args) const;
+	bool CanSpawnDataLayerBrowser(const FSpawnTabArgs& Args) const;
 	bool HasAnyHostedEditorModeToolkit() const;
 
 	//TSharedRef<SDockTab> SpawnLevelEditorModeTab(const FSpawnTabArgs& Args, FEdMode* EditorMode);
diff --git a/Engine/Source/Editor/LevelEditor/Private/SLevelViewport.cpp b/Engine/Source/Editor/LevelEditor/Private/SLevelViewport.cpp
--- a/Engine/Source/Editor/LevelEditor/Private/SLevelViewport.cpp
+++ b/Engine/Source/Editor/LevelEditor/Private/SLevelViewport.cpp
@@ -78,6 +78,10 @@
 #include "Bookmarks/IBookmarkTypeTools.h"
 #include "Editor/EditorPerformanceSettings.h"
 #include "UnrealWidget.h"
+#include "WorldPartition/WorldPartitionSubsystem.h"
+#include "WorldPartition/DataLayer/WorldDataLayers.h"
+#include "WorldPartition/DataLayer/DataLayer.h"
+#include "DataLayer/DataLayerEditorSubsystem.h"
 
 static const FName LevelEditorName("LevelEditor");
 
@@ -1593,16 +1597,33 @@ void SLevelViewport::BindShowCommands( FUICommandList& OutCommandList )
 
 	// Show Layers
 	{
+		auto CanToggleAllLayers = [this]() { return !UWorld::HasSubsystem<UWorldPartitionSubsystem>(GetWorld()); };
 		// Map 'Show All' and 'Hide All' commands
 		OutCommandList.MapAction(
 			LevelViewportCommands.ShowAllLayers,
-			FExecuteAction::CreateSP( this, &SLevelViewport::OnToggleAllLayers, true ) );
+			FExecuteAction::CreateSP( this, &SLevelViewport::OnToggleAllLayers, true ),
+			FCanExecuteAction::CreateLambda(CanToggleAllLayers));
 
 		OutCommandList.MapAction(
 			LevelViewportCommands.HideAllLayers,
-			FExecuteAction::CreateSP( this, &SLevelViewport::OnToggleAllLayers, false ) );
+			FExecuteAction::CreateSP( this, &SLevelViewport::OnToggleAllLayers, false ),
+			FCanExecuteAction::CreateLambda(CanToggleAllLayers));
 	}
+	// Show DataLayers
+	{
+		auto CanToggleAllDataLayers = [this]() { return UWorld::HasSubsystem<UWorldPartitionSubsystem>(GetWorld()); };
+		// Map 'Show All' and 'Hide All' commands
+		OutCommandList.MapAction(
+			LevelViewportCommands.ShowAllDataLayers,
+			FExecuteAction::CreateSP(this, &SLevelViewport::OnToggleAllDataLayers, true),
+			FCanExecuteAction::CreateLambda(CanToggleAllDataLayers));
 
+		OutCommandList.MapAction(
+			LevelViewportCommands.HideAllDataLayers,
+			FExecuteAction::CreateSP(this, &SLevelViewport::OnToggleAllDataLayers, false),
+			FCanExecuteAction::CreateLambda(CanToggleAllDataLayers));
+	}
+	
 	// Show Sprite Categories
 	{
 		// Map 'Show All' and 'Hide All' commands
@@ -2068,6 +2089,63 @@ bool SLevelViewport::IsLayerVisible( FName LayerName ) const
 	return LevelViewportClient->ViewHiddenLayers.Find(LayerName) == INDEX_NONE;
 }
 
+/** Called when a user selects show or hide all from the DataLayers visibility menu. **/
+void SLevelViewport::OnToggleAllDataLayers(bool bVisible)
+{
+	UDataLayerEditorSubsystem* DataLayerSubsystem = UDataLayerEditorSubsystem::Get();
+	if (bVisible)
+	{
+		// Clear all hidden DataLayers
+		LevelViewportClient->ViewHiddenDataLayers.Empty();
+	}
+	else
+	{
+		TArray<FName> AllDataLayerNames;
+		if (const AWorldDataLayers* WorldDataLayers = AWorldDataLayers::Get(GetWorld()))
+		{
+			WorldDataLayers->ForEachDataLayer([&AllDataLayerNames](UDataLayer* DataLayer)
+			{
+				AllDataLayerNames.Add(DataLayer->GetFName());
+				return true;
+			});
+		}
+		// Hide them all
+		LevelViewportClient->ViewHiddenDataLayers = AllDataLayerNames;
+	}
+
+	// Update actor visibility for this view
+	DataLayerSubsystem->UpdatePerViewVisibility(LevelViewportClient.Get());
+	LevelViewportClient->Invalidate();
+}
+
+/** Called when the user toggles a layer from Layers sub-menu. **/
+void SLevelViewport::ToggleShowDataLayer(FName DataLayerName)
+{
+	int32 HiddenIndex = LevelViewportClient->ViewHiddenDataLayers.Find(DataLayerName);
+	if (HiddenIndex == INDEX_NONE)
+	{
+		LevelViewportClient->ViewHiddenDataLayers.Add(DataLayerName);
+	}
+	else
+	{
+		LevelViewportClient->ViewHiddenDataLayers.RemoveAt(HiddenIndex);
+	}
+
+	// Update actor visibility for this view
+	UDataLayerEditorSubsystem* DataLayerEditorSubsystem = UDataLayerEditorSubsystem::Get();
+	if (const UDataLayer* DataLayer = DataLayerEditorSubsystem->GetDataLayerFromName(DataLayerName))
+	{
+		DataLayerEditorSubsystem->UpdatePerViewVisibility(LevelViewportClient.Get(), DataLayer);
+		LevelViewportClient->Invalidate();
+	}
+}
+
+/** Called to determine if a DataLayer is visible. **/
+bool SLevelViewport::IsDataLayerVisible(FName DataLayerName) const
+{
+	return LevelViewportClient->ViewHiddenDataLayers.Find(DataLayerName) == INDEX_NONE;
+}
+
 void SLevelViewport::ToggleShowFoliageType(TWeakObjectPtr<UFoliageType> InFoliageType)
 {
 	UFoliageType* FoliageType = InFoliageType.Get();
diff --git a/Engine/Source/Editor/LevelEditor/Private/SLevelViewportToolBar.cpp b/Engine/Source/Editor/LevelEditor/Private/SLevelViewportToolBar.cpp
--- a/Engine/Source/Editor/LevelEditor/Private/SLevelViewportToolBar.cpp
+++ b/Engine/Source/Editor/LevelEditor/Private/SLevelViewportToolBar.cpp
@@ -41,6 +41,9 @@
 #include "ToolMenu.h"
 #include "CustomStaticScreenPercentage.h"
 #include "CustomEditorStaticScreenPercentage.h"
+#include "WorldPartition/WorldPartitionSubsystem.h"
+#include "WorldPartition/DataLayer/WorldDataLayers.h"
+#include "WorldPartition/DataLayer/DataLayer.h"
 
 #define LOCTEXT_NAMESPACE "LevelViewportToolBar"
 
@@ -1150,6 +1153,9 @@ void SLevelViewportToolBar::FillShowMenu(UToolMenu* Menu) const
 		FText ShowAllLabel = LOCTEXT("ShowAllLabel", "Show All");
 		FText HideAllLabel = LOCTEXT("HideAllLabel", "Hide All");
 
+		FLevelEditorViewportClient& ViewClient = Viewport.Pin()->GetLevelViewportClient();
+		const UWorld* World = ViewClient.GetWorld();
+
 		{
 			FToolMenuSection& Section = Menu->AddSection("LevelViewportEditorShow", LOCTEXT("EditorShowHeader", "Editor"));
 			// Show Volumes sub-menu
@@ -1168,10 +1174,18 @@ void SLevelViewportToolBar::FillShowMenu(UToolMenu* Menu) const
 			}
 
 			// Show Layers sub-menu is dynamically generated when the user enters 'show' menu
+			if (!UWorld::HasSubsystem<UWorldPartitionSubsystem>(World))
 			{
 				Section.AddSubMenu("ShowLayers", LOCTEXT("ShowLayersMenu", "Layers"), LOCTEXT("ShowLayersMenu_ToolTip", "Show layers flags"),
 					FNewToolMenuDelegate::CreateStatic(&SLevelViewportToolBar::FillShowLayersMenu, Viewport));
 			}
+			else
+			{
+				const FSlateIcon DataLayersIcon(FEditorStyle::GetStyleSetName(), "LevelEditor.Tabs.DataLayers");
+				// Show DataLayers sub-menu is dynamically generated when the user enters 'show' menu
+				Section.AddSubMenu("ShowDataLayers", LOCTEXT("ShowDataLayersMenu", "Data Layers"), LOCTEXT("ShowDataLayersMenu_ToolTip", "Show Data Layers flags"),
+					FNewToolMenuDelegate::CreateStatic(&SLevelViewportToolBar::FillShowDataLayersMenu, Viewport), false, DataLayersIcon);
+			}
 
 			// Show Sprites sub-menu
 			{
@@ -1401,6 +1415,36 @@ void SLevelViewportToolBar::FillShowLayersMenu(UToolMenu* Menu, TWeakPtr<class S
 	}
 }
 
+void SLevelViewportToolBar::FillShowDataLayersMenu(UToolMenu* Menu, TWeakPtr<class SLevelViewport> Viewport)
+{
+	{
+		FToolMenuSection& Section = Menu->AddSection("LevelViewportDataLayers");
+		Section.AddMenuEntry(FLevelViewportCommands::Get().ShowAllDataLayers, LOCTEXT("ShowAllLabel", "Show All"));
+		Section.AddMenuEntry(FLevelViewportCommands::Get().HideAllDataLayers, LOCTEXT("HideAllLabel", "Hide All"));
+	}
+
+	if (Viewport.IsValid())
+	{
+		const FSlateIcon DataLayerIcon(FEditorStyle::GetStyleSetName(), "DataLayer.Icon16x");
+		TSharedRef<SLevelViewport> ViewportRef = Viewport.Pin().ToSharedRef();
+		FToolMenuSection& Section = Menu->AddSection("LevelViewportDataLayers2");
+		// Get all the DataLayers and create an entry for each of them
+		if (const AWorldDataLayers* WorldDataLayers = AWorldDataLayers::Get(ViewportRef->GetWorld()))
+		{
+			WorldDataLayers->ForEachDataLayer([&Section, &ViewportRef, &DataLayerIcon](UDataLayer* DataLayer)
+			{
+				const FName DataLayerName = DataLayer->GetFName();
+				const FName DataLayerLabel = DataLayer->GetDataLayerLabel();
+				FUIAction Action(FExecuteAction::CreateSP(ViewportRef, &SLevelViewport::ToggleShowDataLayer, DataLayerName),
+					FCanExecuteAction(),
+					FIsActionChecked::CreateSP(ViewportRef, &SLevelViewport::IsDataLayerVisible, DataLayerName));
+				Section.AddMenuEntry(NAME_None, FText::FromName(DataLayerLabel), FText::GetEmpty(), DataLayerIcon, Action, EUserInterfaceActionType::ToggleButton);
+				return true;
+			});
+		}
+	}
+}
+
 static TMap<FName, TArray<UFoliageType*>> GroupFoliageByOuter(const TArray<UFoliageType*> FoliageList)
 {
 	TMap<FName, TArray<UFoliageType*>> Result;
diff --git a/Engine/Source/Editor/LevelEditor/Private/SLevelViewportToolBar.h b/Engine/Source/Editor/LevelEditor/Private/SLevelViewportToolBar.h
--- a/Engine/Source/Editor/LevelEditor/Private/SLevelViewportToolBar.h
+++ b/Engine/Source/Editor/LevelEditor/Private/SLevelViewportToolBar.h
@@ -249,6 +249,13 @@ private:
 	 */
 	static void FillShowLayersMenu(UToolMenu* Menu, TWeakPtr<class SLevelViewport> Viewport );
 
+	/**
+	 * Generates the toolbar show DataLayers menu content
+	 *
+	 * @param Menu	The tool menu
+	 */
+	static void FillShowDataLayersMenu(UToolMenu* Menu, TWeakPtr<class SLevelViewport> Viewport);
+
 	/**
 	 * Generates 'show foliage types' menu content for a viewport
 	 *
diff --git a/Engine/Source/Editor/LevelEditor/Public/LevelEditor.h b/Engine/Source/Editor/LevelEditor/Public/LevelEditor.h
--- a/Engine/Source/Editor/LevelEditor/Public/LevelEditor.h
+++ b/Engine/Source/Editor/LevelEditor/Public/LevelEditor.h
@@ -43,6 +43,7 @@ public:
 	static const FName LevelEditorSceneOutliner;
 	static const FName LevelEditorStatsViewer;
 	static const FName LevelEditorLayerBrowser;
+	static const FName LevelEditorDataLayerBrowser;
 	static const FName Sequencer;
 	static const FName SequencerGraphEditor;
 	static const FName WorldSettings;
diff --git a/Engine/Source/Editor/LevelEditor/Public/LevelViewportActions.h b/Engine/Source/Editor/LevelEditor/Public/LevelViewportActions.h
--- a/Engine/Source/Editor/LevelEditor/Public/LevelViewportActions.h
+++ b/Engine/Source/Editor/LevelEditor/Public/LevelViewportActions.h
@@ -126,6 +126,12 @@ PRAGMA_ENABLE_DEPRECATION_WARNINGS
 	/** Hides all layers */
 	TSharedPtr< FUICommandInfo > HideAllLayers;
 
+	/** Shows all DataLayers */
+	TSharedPtr< FUICommandInfo > ShowAllDataLayers;
+
+	/** Hides all DataLayers */
+	TSharedPtr< FUICommandInfo > HideAllDataLayers;
+
 	/** Shows all sprite categories */
 	TSharedPtr< FUICommandInfo > ShowAllSprites;
 
diff --git a/Engine/Source/Editor/LevelEditor/Public/SLevelViewport.h b/Engine/Source/Editor/LevelEditor/Public/SLevelViewport.h
--- a/Engine/Source/Editor/LevelEditor/Public/SLevelViewport.h
+++ b/Engine/Source/Editor/LevelEditor/Public/SLevelViewport.h
@@ -203,6 +203,20 @@ public:
 	 */
 	bool IsLayerVisible( FName LayerName ) const;
 
+	/**
+	 * Toggles DataLayer visibility in this viewport
+	 *
+	 * @param DataLayerName		Name of the DataLayer
+	 */
+	void ToggleShowDataLayer(FName DataLayerName);
+
+	/**
+	 * Checks if a DataLayer is visible in this viewport
+	 *
+	 * @param DataLayerName		Name of the DataLayer
+	 */
+	bool IsDataLayerVisible(FName DataLayerName) const;
+
 	/**
 	 * Toggles foliage type visibility in this viewport
 	 *
@@ -503,6 +517,13 @@ private:
 	 */
 	void OnToggleAllLayers( bool bVisible );
 
+	/**
+	 * Toggles all DataLayers visibility
+	 *
+	 * @param Visible			true if DataLayers should be visible, false otherwise
+	 */
+	void OnToggleAllDataLayers( bool bVisible );
+
 	/**
 	 * Toggles all sprite categories visibility
 	 *
diff --git a/Engine/Source/Editor/SceneOutliner/Private/SceneOutlinerActorInfoColumn.cpp b/Engine/Source/Editor/SceneOutliner/Private/SceneOutlinerActorInfoColumn.cpp
--- a/Engine/Source/Editor/SceneOutliner/Private/SceneOutlinerActorInfoColumn.cpp
+++ b/Engine/Source/Editor/SceneOutliner/Private/SceneOutlinerActorInfoColumn.cpp
@@ -4,6 +4,7 @@
 #include "Modules/ModuleManager.h"
 #include "UObject/Package.h"
 #include "Misc/PackageName.h"
+#include "Misc/StringBuilder.h"
 #include "Framework/Application/SlateApplication.h"
 #include "Widgets/Text/STextBlock.h"
 #include "Widgets/Views/SListView.h"
@@ -15,6 +16,8 @@
 #include "ComponentTreeItem.h"
 #include "FolderTreeItem.h"
 #include "WorldTreeItem.h"
+#include "WorldPartition/DataLayer/DataLayer.h"
+#include "WorldPartition/DataLayer/WorldDataLayers.h"
 
 #define LOCTEXT_NAMESPACE "SceneOutlinerActorInfoColumn"
 
@@ -65,6 +68,21 @@ struct FGetInfo
 				}
 				return Result;
 			}
+
+			case SceneOutliner::ECustomColumnMode::DataLayer:
+			{
+				TStringBuilder<128> Builder;
+				for (const UDataLayer* DataLayer : Actor->GetDataLayerObjects())
+				{
+					if (Builder.Len())
+					{
+						Builder += TEXT(", ");
+					}
+					Builder += DataLayer->GetDataLayerLabel().ToString();
+				}
+				return Builder.ToString();
+			}
+
 		case SceneOutliner::ECustomColumnMode::Mobility:
 			{
 				FString Result;
@@ -373,6 +391,10 @@ FText FActorInfoColumn::MakeComboText( const SceneOutliner::ECustomColumnMode::T
 		ModeName = LOCTEXT("CustomColumnMode_Layer", "Layer");
 		break;
 
+	case SceneOutliner::ECustomColumnMode::DataLayer:
+		ModeName = LOCTEXT("CustomColumnMode_DataLayer", "Data Layer");
+		break;
+
 	case SceneOutliner::ECustomColumnMode::Socket:
 		ModeName = LOCTEXT("CustomColumnMode_Socket", "Socket");
 		break;
@@ -420,6 +442,10 @@ FText FActorInfoColumn::MakeComboToolTipText( const SceneOutliner::ECustomColumn
 		ToolTipText = LOCTEXT("CustomColumnModeToolTip_Layer", "Displays the layer each actor is in, and allows you to search by layer name");
 		break;
 
+	case SceneOutliner::ECustomColumnMode::DataLayer:
+		ToolTipText = LOCTEXT("CustomColumnModeToolTip_DataLayer", "Displays the data layers each actor is in, and allows you to search by data layer label");
+		break;
+
 	case SceneOutliner::ECustomColumnMode::Socket:
 		ToolTipText = LOCTEXT("CustomColumnModeToolTip_Socket", "Shows the socket the actor is attached to, and allows you to search by socket name");
 		break;
diff --git a/Engine/Source/Editor/SceneOutliner/Private/SceneOutlinerActorInfoColumn.h b/Engine/Source/Editor/SceneOutliner/Private/SceneOutlinerActorInfoColumn.h
--- a/Engine/Source/Editor/SceneOutliner/Private/SceneOutlinerActorInfoColumn.h
+++ b/Engine/Source/Editor/SceneOutliner/Private/SceneOutlinerActorInfoColumn.h
@@ -36,6 +36,9 @@ namespace SceneOutliner
 			/** Layer */
 			Layer,
 
+			/** Data Layer */
+			DataLayer,
+
 			/** The socket the actor is attached to. */
 			Socket,
 
diff --git a/Engine/Source/Editor/UnrealEd/Private/DragTool_BoxSelect.cpp b/Engine/Source/Editor/UnrealEd/Private/DragTool_BoxSelect.cpp
--- a/Engine/Source/Editor/UnrealEd/Private/DragTool_BoxSelect.cpp
+++ b/Engine/Source/Editor/UnrealEd/Private/DragTool_BoxSelect.cpp
@@ -177,6 +177,7 @@ void FDragTool_ActorBoxSelect::EndDrag()
 		bool bSelectionChanged = false;
 		UWorld* IteratorWorld = GWorld;
 		const TArray<FName>& HiddenLayers = LevelViewportClient->ViewHiddenLayers;
+		const TArray<FName>& HiddenDataLayers = LevelViewportClient->ViewHiddenDataLayers;
 		for( FActorIterator It(IteratorWorld); It; ++It )
 		{
 			AActor* Actor = *It;
@@ -191,6 +192,10 @@ void FDragTool_ActorBoxSelect::EndDrag()
 					break;
 				}
 			}
+			if (bActorIsVisible)
+			{
+				bActorIsVisible = !Actor->HasAnyOfDataLayers(HiddenDataLayers);
+			}
 
 			// Select the actor if we need to
 			if( bActorIsVisible && IntersectsBox( *Actor, SelBBox, bStrictDragSelection ) )
diff --git a/Engine/Source/Editor/UnrealEd/Private/LevelEditorViewport.cpp b/Engine/Source/Editor/UnrealEd/Private/LevelEditorViewport.cpp
--- a/Engine/Source/Editor/UnrealEd/Private/LevelEditorViewport.cpp
+++ b/Engine/Source/Editor/UnrealEd/Private/LevelEditorViewport.cpp
@@ -1713,6 +1713,7 @@ void FTrackingTransaction::OnEditorSelectionChanged(UObject* NewSelection)
 FLevelEditorViewportClient::FLevelEditorViewportClient(const TSharedPtr<SLevelViewport>& InLevelViewport)
 	: FEditorViewportClient(&GLevelEditorModeTools(), nullptr, StaticCastSharedPtr<SEditorViewport>(InLevelViewport))
 	, ViewHiddenLayers()
+	, ViewHiddenDataLayers()
 	, VolumeActorVisibility()
 	, LastEditorViewLocation( FVector::ZeroVector )
 	, LastEditorViewRotation( FRotator::ZeroRotator )
diff --git a/Engine/Source/Editor/UnrealEd/Public/LevelEditorViewport.h b/Engine/Source/Editor/UnrealEd/Public/LevelEditorViewport.h
--- a/Engine/Source/Editor/UnrealEd/Public/LevelEditorViewport.h
+++ b/Engine/Source/Editor/UnrealEd/Public/LevelEditorViewport.h
@@ -724,6 +724,9 @@ public:
 	/** List of layers that are hidden in this view */
 	TArray<FName>			ViewHiddenLayers;
 
+	/** List of DataLayers that are hidden in this view */
+	TArray<FName>			ViewHiddenDataLayers;
+
 	/** Special volume actor visibility settings. Each bit represents a visibility state for a specific volume class. 1 = visible, 0 = hidden */
 	TBitArray<>				VolumeActorVisibility;
 
diff --git a/Engine/Source/Runtime/Core/Public/UObject/UE5MainStreamObjectVersion.h b/Engine/Source/Runtime/Core/Public/UObject/UE5MainStreamObjectVersion.h
--- a/Engine/Source/Runtime/Core/Public/UObject/UE5MainStreamObjectVersion.h
+++ b/Engine/Source/Runtime/Core/Public/UObject/UE5MainStreamObjectVersion.h
@@ -34,6 +34,9 @@ struct CORE_API FUE5MainStreamObjectVersion
 		// Added a new configurable BlendProfileMode that the user can setup to control the behavior of blend profiles.
 		AnimationAddedBlendProfileModes,
 
+		// Serialize DataLayers in WorldPartitionActorDesc
+		WorldPartitionActorDescSerializeDataLayers,
+
 		// -----<new versions can be added above this line>-------------------------------------------------
 		VersionPlusOne,
 		LatestVersion = VersionPlusOne - 1
diff --git a/Engine/Source/Runtime/Engine/Classes/GameFramework/Actor.h b/Engine/Source/Runtime/Engine/Classes/GameFramework/Actor.h
--- a/Engine/Source/Runtime/Engine/Classes/GameFramework/Actor.h
+++ b/Engine/Source/Runtime/Engine/Classes/GameFramework/Actor.h
@@ -18,6 +18,10 @@
 #include "Misc/ITransaction.h"
 #include "Engine/Level.h"
 
+#if WITH_EDITOR
+#include "WorldPartition/DataLayer/ActorDataLayer.h"
+#endif
+
 #include "Actor.generated.h"
 
 class AActor;
@@ -32,6 +36,7 @@ class UPrimitiveComponent;
 struct FAttachedActorInfo;
 struct FNetViewer;
 struct FNetworkObjectInfo;
+class UDataLayer;
 
 #if WITH_EDITOR
 // @todo_ow this is temporary and will be removed
@@ -763,6 +768,10 @@ protected:
 	UPROPERTY(VisibleAnywhere, AdvancedDisplay, Category=Actor, NonPIEDuplicateTransient, TextExportTransient, NonTransactional)
 	FGuid ActorGuid;
 
+	/** DataLayers the actor belongs to.*/
+	UPROPERTY(EditAnywhere, AdvancedDisplay, Category = Actor)
+	TArray<FActorDataLayer> DataLayers;
+
 public:
 	/** The editor-only group this actor is a part of. */
 	UPROPERTY(Transient)
@@ -998,6 +1007,21 @@ public:
 		return Cast<T>(GetInstigatorController());
 	}
 
+#if WITH_EDITOR
+	//~=============================================================================
+	// DataLayers functions.
+	bool AddDataLayer(const UDataLayer* DataLayer);
+	bool RemoveDataLayer(const UDataLayer* DataLayer);
+	bool ContainsDataLayer(const UDataLayer* DataLayer) const;
+	bool HasDataLayers() const;
+	bool HasValidDataLayers() const;
+	bool HasAllDataLayers(const TArray<const UDataLayer*>& DataLayers) const;
+	bool HasAnyOfDataLayers(const TArray<FName>& DataLayerNames) const;
+	TArray<FName> GetDataLayerNames() const;
+	TArray<const UDataLayer*> GetDataLayerObjects() const;
+	bool IsPropertyChangedAffectingDataLayers(FPropertyChangedEvent& PropertyChangedEvent) const;
+	void FixupDataLayers();
+#endif
 
 	//~=============================================================================
 	// General functions.
diff --git a/Engine/Source/Runtime/Engine/Private/Actor.cpp b/Engine/Source/Runtime/Engine/Private/Actor.cpp
--- a/Engine/Source/Runtime/Engine/Private/Actor.cpp
+++ b/Engine/Source/Runtime/Engine/Private/Actor.cpp
@@ -787,6 +787,12 @@ void AActor::Serialize(FArchive& Ar)
 		}
 	}
 
+	// Fixup DataLayers
+	if (Ar.IsSaving() && !Ar.IsObjectReferenceCollector() && Ar.IsPersistent())
+	{
+		FixupDataLayers();
+	}
+
 	// When duplicating for PIE all components need to be gathered up and duplicated even if there are no other property references to them
 	// otherwise we can end up with Attach Parents that do not get redirected to the correct component. However, if there is a transient component
 	// we'll let that drop
@@ -886,6 +892,10 @@ void AActor::PostLoad()
 			UE_LOG(LogActor, Log,  TEXT("Loaded Actor (%s) with IsPendingKill() == true"), *GetName() );
 		}
 	}
+
+	// Fixup DataLayers
+	FixupDataLayers();
+
 #endif // WITH_EDITORONLY_DATA
 
 	// Since the actor is being loading, it finished spawning by definition when it was originally spawned, so set to true now
diff --git a/Engine/Source/Runtime/Engine/Private/ActorEditor.cpp b/Engine/Source/Runtime/Engine/Private/ActorEditor.cpp
--- a/Engine/Source/Runtime/Engine/Private/ActorEditor.cpp
+++ b/Engine/Source/Runtime/Engine/Private/ActorEditor.cpp
@@ -13,6 +13,8 @@
 #include "AI/NavigationSystemBase.h"
 #include "Engine/BlueprintGeneratedClass.h"
 #include "WorldPartition/WorldPartitionSubsystem.h"
+#include "WorldPartition/DataLayer/WorldDataLayers.h"
+#include "WorldPartition/DataLayer/DataLayer.h"
 #include "EditorSupportDelegates.h"
 #include "Logging/TokenizedMessage.h"
 #include "Logging/MessageLog.h"
@@ -83,6 +85,11 @@ void AActor::PostEditChangeProperty(FPropertyChangedEvent& PropertyChangedEvent)
 {
 	FProperty* MemberPropertyThatChanged = PropertyChangedEvent.MemberProperty;
 	const FName MemberPropertyName = MemberPropertyThatChanged != NULL ? MemberPropertyThatChanged->GetFName() : NAME_None;
+
+	if (IsPropertyChangedAffectingDataLayers(PropertyChangedEvent))
+	{
+		FixupDataLayers();
+	}
 	
 	const bool bTransformationChanged = (MemberPropertyName == Name_RelativeLocation || MemberPropertyName == Name_RelativeRotation || MemberPropertyName == Name_RelativeScale3D);
 
@@ -1119,6 +1126,180 @@ EDataValidationResult AActor::IsDataValid(TArray<FText>& ValidationErrors)
 
 	return Result;
 }
+
+//---------------------------------------------------------------------------
+// DataLayers (begin)
+
+bool AActor::AddDataLayer(const UDataLayer* DataLayer)
+{
+	bool bActorWasModified = false;
+	if (DataLayer && !ContainsDataLayer(DataLayer))
+	{
+		if (!bActorWasModified)
+		{
+			Modify();
+			bActorWasModified = true;
+		}
+
+		DataLayers.Emplace(DataLayer->GetFName());
+	}
+	return bActorWasModified;
+}
+
+bool AActor::RemoveDataLayer(const UDataLayer* DataLayer)
+{
+	bool bActorWasModified = false;
+	if (ContainsDataLayer(DataLayer))
+	{
+		if (!bActorWasModified)
+		{
+			Modify();
+			bActorWasModified = true;
+		}
+
+		DataLayers.Remove(FActorDataLayer(DataLayer->GetFName()));
+	}
+	return bActorWasModified;
+}
+
+bool AActor::ContainsDataLayer(const UDataLayer* DataLayer) const
+{
+	return DataLayer && DataLayers.Contains(FActorDataLayer(DataLayer->GetFName()));
+}
+
+bool AActor::HasDataLayers() const
+{
+	return DataLayers.Num() > 0;
+}
+
+bool AActor::HasValidDataLayers() const
+{
+	if (const AWorldDataLayers* WorldDataLayers = AWorldDataLayers::Get(GetWorld()))
+	{
+		for (const FActorDataLayer& DataLayer : DataLayers)
+		{
+			if (const UDataLayer* DataLayerObject = WorldDataLayers->GetDataLayerFromName(DataLayer.Name))
+			{
+				return true;
+			}
+		}
+	}
+	return false;
+}
+
+bool AActor::HasAllDataLayers(const TArray<const UDataLayer*>& InDataLayers) const
+{
+	if (DataLayers.Num() < InDataLayers.Num())
+	{
+		return false;
+	}
+
+	for (const UDataLayer* DataLayer : InDataLayers)
+	{
+		if (!ContainsDataLayer(DataLayer))
+		{
+			return false;
+		}
+	}
+	return true;
+}
+
+TArray<FName> AActor::GetDataLayerNames() const
+{
+	TArray<FName> OutDataLayerNames;
+	OutDataLayerNames.Reserve(DataLayers.Num());
+
+	for (const UDataLayer* DataLayer : GetDataLayerObjects())
+	{
+		OutDataLayerNames.Add(DataLayer->GetFName());
+	}
+
+	return MoveTemp(OutDataLayerNames);
+}
+
+TArray<const UDataLayer*> AActor::GetDataLayerObjects() const
+{
+	TArray<const UDataLayer*> OutDataLayers;
+	OutDataLayers.Reserve(DataLayers.Num());
+	
+	if (const AWorldDataLayers* WorldDataLayers = AWorldDataLayers::Get(GetWorld()))
+	{
+		for (const FActorDataLayer& DataLayer : DataLayers)
+		{
+			if (const UDataLayer* DataLayerObject = WorldDataLayers->GetDataLayerFromName(DataLayer.Name))
+			{
+				OutDataLayers.Add(DataLayerObject);
+			}
+		}
+	}
+
+	return MoveTemp(OutDataLayers);
+}
+
+bool AActor::HasAnyOfDataLayers(const TArray<FName>& DataLayerNames) const
+{
+	for (const FActorDataLayer& DataLayer : DataLayers)
+	{
+		if (DataLayerNames.Contains(DataLayer.Name))
+		{
+			return true;
+		}
+	}
+	return false;
+}
+
+void AActor::FixupDataLayers()
+{
+	if (const AWorldDataLayers* WorldDataLayers = AWorldDataLayers::Get(GetWorld()))
+	{
+		TSet<FName> ExistingDataLayers;
+		for (int32 Index = 0; Index < DataLayers.Num();)
+		{
+			const FName& DataLayer = DataLayers[Index].Name;
+			if (!WorldDataLayers->GetDataLayerFromName(DataLayer) || ExistingDataLayers.Contains(DataLayer))
+			{
+				DataLayers.RemoveAtSwap(Index);
+			}
+			else
+			{
+				ExistingDataLayers.Add(DataLayer);
+				++Index;
+			}
+		}
+	}
+}
+
+bool AActor::IsPropertyChangedAffectingDataLayers(FPropertyChangedEvent& PropertyChangedEvent) const
+{
+	if (PropertyChangedEvent.Property != nullptr)
+	{
+		FProperty* MemberPropertyThatChanged = PropertyChangedEvent.MemberProperty;
+		const FName MemberPropertyName = MemberPropertyThatChanged != NULL ? MemberPropertyThatChanged->GetFName() : NAME_None;
+
+		static const FName NAME_DataLayers = GET_MEMBER_NAME_CHECKED(AActor, DataLayers);
+		static const FName NAME_FActorDataLayerName = GET_MEMBER_NAME_CHECKED(FActorDataLayer, Name);
+
+		if (MemberPropertyName == NAME_DataLayers &&
+			PropertyChangedEvent.ChangeType == EPropertyChangeType::ValueSet &&
+			PropertyChangedEvent.Property->GetFName() == NAME_FActorDataLayerName)
+		{
+			return true;
+		}
+		else
+		{
+			const FName PropertyName = PropertyChangedEvent.GetPropertyName();
+			if (PropertyName == NAME_DataLayers && (PropertyChangedEvent.ChangeType == EPropertyChangeType::ValueSet))
+			{
+				return true;
+			}
+		}
+	}
+	return false;
+}
+
+// DataLayers (end)
+//---------------------------------------------------------------------------
+
 #undef LOCTEXT_NAMESPACE
 
 #endif // WITH_EDITOR
diff --git a/Engine/Source/Runtime/Engine/Private/WorldPartition/DataLayer/DataLayer.cpp b/Engine/Source/Runtime/Engine/Private/WorldPartition/DataLayer/DataLayer.cpp
new file mode 100644
--- /dev/null
+++ b/Engine/Source/Runtime/Engine/Private/WorldPartition/DataLayer/DataLayer.cpp
@@ -0,0 +1,49 @@
+// Copyright Epic Games, Inc. All Rights Reserved.
+
+/*=============================================================================
+	DataLayer.cpp: UDataLayer class implementation
+=============================================================================*/
+
+#include "WorldPartition/DataLayer/DataLayer.h"
+
+UDataLayer::UDataLayer(const FObjectInitializer& ObjectInitializer)
+: Super(ObjectInitializer)
+#if WITH_EDITORONLY_DATA
+, DataLayerLabel(GetFName())
+, bIsVisible(true)
+, bIsDynamicallyLoaded(true)
+#endif
+{
+}
+
+#if WITH_EDITOR
+void UDataLayer::SetDataLayerLabel(FName InDataLayerLabel)
+{
+	if (DataLayerLabel != InDataLayerLabel)
+	{
+		Modify();
+		AWorldDataLayers* WorldDataLayers = GetOuterAWorldDataLayers();
+		check(!WorldDataLayers || !WorldDataLayers->GetDataLayerFromLabel(InDataLayerLabel))
+		DataLayerLabel = InDataLayerLabel;
+	}
+}
+
+void UDataLayer::SetVisible(bool bInIsVisible)
+{
+	if (bIsVisible != bInIsVisible)
+	{
+		Modify();
+		bIsVisible = bInIsVisible;
+	}
+}
+
+void UDataLayer::SetIsDynamicallyLoaded(bool bInIsDynamicallyLoaded)
+{
+	if (bIsDynamicallyLoaded != bInIsDynamicallyLoaded)
+	{
+		Modify();
+		bIsDynamicallyLoaded = bInIsDynamicallyLoaded;
+	}
+}
+
+#endif
\ No newline at end of file
diff --git a/Engine/Source/Runtime/Engine/Private/WorldPartition/DataLayer/DataLayerHelper.cpp b/Engine/Source/Runtime/Engine/Private/WorldPartition/DataLayer/DataLayerHelper.cpp
--- a/Engine/Source/Runtime/Engine/Private/WorldPartition/DataLayer/DataLayerHelper.cpp
+++ b/Engine/Source/Runtime/Engine/Private/WorldPartition/DataLayer/DataLayerHelper.cpp
@@ -1,17 +1,22 @@
 // Copyright Epic Games, Inc. All Rights Reserved.
 
 #include "WorldPartition/DataLayer/DataLayerHelper.h"
+#include "WorldPartition/DataLayer/DataLayer.h"
 #include "Misc/HashBuilder.h"
+#include "Algo/Transform.h"
 
-uint32 FDataLayersHelper::ComputeDataLayerID(const TArray<FName>& InDataLayers)
+#if WITH_EDITOR
+uint32 FDataLayersHelper::ComputeDataLayerID(const TArray<const UDataLayer*>& InDataLayers)
 {
 	if (InDataLayers.Num())
 	{
-		TArray<FName> SortedDataLayers = InDataLayers;
+		TArray<FName> SortedDataLayers;
+		Algo::TransformIf(InDataLayers, SortedDataLayers, [](const UDataLayer* Item) { return Item->IsDynamicallyLoaded(); }, [](const UDataLayer* Item) { return Item->GetFName(); });
 		SortedDataLayers.Sort([](const FName& A, const FName& B) { return A.FastLess(B); });
 		FHashBuilder HashBuilder;
 		HashBuilder << SortedDataLayers;
 		return HashBuilder.GetHash();
 	}
 	return FDataLayersHelper::NoDataLayerID;
-}
\ No newline at end of file
+}
+#endif
\ No newline at end of file
diff --git a/Engine/Source/Runtime/Engine/Private/WorldPartition/DataLayer/DataLayerSubsystem.cpp b/Engine/Source/Runtime/Engine/Private/WorldPartition/DataLayer/DataLayerSubsystem.cpp
--- a/Engine/Source/Runtime/Engine/Private/WorldPartition/DataLayer/DataLayerSubsystem.cpp
+++ b/Engine/Source/Runtime/Engine/Private/WorldPartition/DataLayer/DataLayerSubsystem.cpp
@@ -1,12 +1,26 @@
 // Copyright Epic Games, Inc. All Rights Reserved.
 
 #include "WorldPartition/DataLayer/DataLayerSubsystem.h"
-#include "Engine/World.h"
+#include "WorldPartition/DataLayer/WorldDataLayers.h"
+#include "WorldPartition/DataLayer/DataLayer.h"
 #include "Engine/Engine.h"
+#include "Engine/World.h"
 
 UDataLayerSubsystem::UDataLayerSubsystem()
 {}
 
+const UDataLayer* UDataLayerSubsystem::GetDataLayerFromLabel(const FName& InDataLayerLabel) const
+{
+	const AWorldDataLayers* WorldDataLayers = AWorldDataLayers::Get(GetWorld());
+	return WorldDataLayers ? WorldDataLayers->GetDataLayerFromLabel(InDataLayerLabel) : nullptr;
+}
+
+const UDataLayer* UDataLayerSubsystem::GetDataLayerFromName(const FName& InDataLayerName) const
+{
+	const AWorldDataLayers* WorldDataLayers = AWorldDataLayers::Get(GetWorld());
+	return WorldDataLayers ? WorldDataLayers->GetDataLayerFromName(InDataLayerName) : nullptr;
+}
+
 bool UDataLayerSubsystem::ShouldCreateSubsystem(UObject* Outer) const
 {
 	if (!Super::ShouldCreateSubsystem(Outer))
@@ -22,33 +36,38 @@ bool UDataLayerSubsystem::ShouldCreateSubsystem(UObject* Outer) const
 	return false;
 }
 
-void UDataLayerSubsystem::ActivateDataLayer(const FName& InDataLayer, bool bActivate)
+void UDataLayerSubsystem::ActivateDataLayer(const FName& InDataLayerName, bool bActivate)
+{
+	// First validate that DataLayer is known by the world
+	if (!GetDataLayerFromName(InDataLayerName))
 {
-	const bool bIsLayerActive = IsDataLayerActive(InDataLayer);
+		return;
+	}
+
+	const bool bIsLayerActive = IsDataLayerActive(InDataLayerName);
 	if (bIsLayerActive != bActivate)
 	{
 		if (bActivate)
 		{
-			//@todo_ow: Once DataLayer is an asset, validate it exists
-			ActiveDataLayers.Add(InDataLayer);
+			ActiveDataLayerNames.Add(InDataLayerName);
 		}
 		else
 		{
-			ActiveDataLayers.Remove(InDataLayer);
+			ActiveDataLayerNames.Remove(InDataLayerName);
 		}
 	}
 }
 
-bool UDataLayerSubsystem::IsDataLayerActive(const FName& InDataLayer) const
+bool UDataLayerSubsystem::IsDataLayerActive(const FName& InDataLayerName) const
 {
-	return ActiveDataLayers.Contains(InDataLayer);
+	return ActiveDataLayerNames.Contains(InDataLayerName);
 }
 
-bool UDataLayerSubsystem::IsAnyDataLayerActive(const TArray<FName>& InDataLayers) const
+bool UDataLayerSubsystem::IsAnyDataLayerActive(const TArray<FName>& InDataLayerNames) const
 {
-	for (const FName& DataLayer : InDataLayers)
+	for (FName DataLayerName : InDataLayerNames)
 	{
-		if (IsDataLayerActive(DataLayer))
+		if (IsDataLayerActive(DataLayerName))
 		{
 			return true;
 		}
@@ -58,20 +77,22 @@ bool UDataLayerSubsystem::IsAnyDataLayerActive(const TArray<FName>& InDataLayers
 
 FAutoConsoleCommand UDataLayerSubsystem::ToggleDataLayerActivation(
 	TEXT("wp.Runtime.ToggleDataLayerActivation"),
-	TEXT("Toggles a DataLayer. Args [DataLayerName]"),
+	TEXT("Toggles a DataLayer. Args [DataLayerLabel]"),
 	FConsoleCommandWithArgsDelegate::CreateLambda([](const TArray<FString>& Args)
 	{
 		if (Args.Num() == 1)
 		{
-			FName DataLayer = FName(Args[0]);
+			FName DataLayerLabel = FName(Args[0]);
 			for (const FWorldContext& Context : GEngine->GetWorldContexts())
 			{
 				UWorld* World = Context.World();
 				if (World && World->IsGameWorld())
 				{
-					if (UDataLayerSubsystem* DataLayerSubsystem = World->GetSubsystem<UDataLayerSubsystem>())
+					UDataLayerSubsystem* DataLayerSubsystem = World->GetSubsystem<UDataLayerSubsystem>();
+					if (const UDataLayer* DataLayer = DataLayerSubsystem ? DataLayerSubsystem->GetDataLayerFromLabel(DataLayerLabel) : nullptr)
 					{
-						DataLayerSubsystem->ActivateDataLayer(DataLayer, !DataLayerSubsystem->IsDataLayerActive(DataLayer));
+						FName DataLayerName = DataLayer->GetFName();
+						DataLayerSubsystem->ActivateDataLayer(DataLayerName, !DataLayerSubsystem->IsDataLayerActive(DataLayerName));
 					}
 				}
 			}
diff --git a/Engine/Source/Runtime/Engine/Private/WorldPartition/DataLayer/WorldDataLayers.cpp b/Engine/Source/Runtime/Engine/Private/WorldPartition/DataLayer/WorldDataLayers.cpp
new file mode 100644
--- /dev/null
+++ b/Engine/Source/Runtime/Engine/Private/WorldPartition/DataLayer/WorldDataLayers.cpp
@@ -0,0 +1,169 @@
+// Copyright Epic Games, Inc. All Rights Reserved.
+
+/*=============================================================================
+	WorldDataLayers.cpp: AWorldDataLayers class implementation
+=============================================================================*/
+
+#include "WorldPartition/DataLayer/WorldDataLayers.h"
+#include "WorldPartition/DataLayer/DataLayer.h"
+#include "EngineUtils.h"
+
+#define LOCTEXT_NAMESPACE "WorldDataLayers"
+
+AWorldDataLayers::AWorldDataLayers(const FObjectInitializer& ObjectInitializer)
+	: Super(ObjectInitializer)
+{
+}
+
+const AWorldDataLayers* AWorldDataLayers::Get(UWorld* World)
+{
+	if (World)
+	{
+		for (AWorldDataLayers* Actor : TActorRange<AWorldDataLayers>(World))
+		{
+			if (Actor && !Actor->IsPendingKill())
+			{
+				return Actor;
+			}
+		}
+	}
+	return nullptr;
+}
+
+#if WITH_EDITOR
+FName AWorldDataLayers::GenerateUniqueDataLayerLabel(const FName& InDataLayerLabel) const
+{
+	int32 DataLayerIndex = 0;
+	FName UniqueNewDataLayerLabel = InDataLayerLabel;
+	while (GetDataLayerFromLabel(UniqueNewDataLayerLabel))
+	{
+		UniqueNewDataLayerLabel = FName(*FString::Printf(TEXT("%s%d"), *InDataLayerLabel.ToString(), ++DataLayerIndex));
+	};
+	return UniqueNewDataLayerLabel;
+}
+
+UDataLayer* AWorldDataLayers::CreateDataLayer()
+{
+	Modify();
+
+	// Make sure new DataLayer name (not label) is unique and never re-used so that actors still referencing on deleted DataLayer's don't get valid again.
+	const FName DataLayerUniqueName = *FString::Format(TEXT("DataLayer_{0}"), { FGuid::NewGuid().ToString()});
+	UDataLayer* NewDataLayer = NewObject<UDataLayer>(this, DataLayerUniqueName, RF_Transactional);
+	check(NewDataLayer != NULL);
+	const FName DefaultLabel = TEXT("DataLayer");
+	FName DataLayerLabel = GenerateUniqueDataLayerLabel(DefaultLabel);
+	NewDataLayer->SetDataLayerLabel(DataLayerLabel);
+	NewDataLayer->SetVisible(true);
+	WorldDataLayers.Add(NewDataLayer);
+	check(GetDataLayerFromName(NewDataLayer->GetFName()));
+	return NewDataLayer;
+}
+
+bool AWorldDataLayers::RemoveDataLayers(const TArray<const UDataLayer*>& InDataLayers)
+{
+	bool bIsModified = false;
+	for (const UDataLayer* DataLayer : InDataLayers)
+	{
+		if (ContainsDataLayer(DataLayer))
+		{
+			Modify();
+			WorldDataLayers.Remove(const_cast<UDataLayer*>(DataLayer));
+			bIsModified = true;
+		}
+	}
+	return bIsModified;
+}
+
+bool AWorldDataLayers::RemoveDataLayer(const UDataLayer* InDataLayer)
+{
+	if (ContainsDataLayer(InDataLayer))
+	{
+		Modify();
+		WorldDataLayers.Remove(const_cast<UDataLayer*>(InDataLayer));
+		return true;
+	}
+	return false;
+}
+
+#endif
+
+bool AWorldDataLayers::ContainsDataLayer(const UDataLayer* InDataLayer) const
+{
+	return WorldDataLayers.Contains(InDataLayer);
+}
+
+const UDataLayer* AWorldDataLayers::GetDataLayerFromName(const FName& InDataLayerName) const
+{
+#if WITH_EDITOR	
+	for (UDataLayer* DataLayer : WorldDataLayers)
+	{
+		if (DataLayer->GetFName() == InDataLayerName)
+		{
+			return DataLayer;
+		}
+	}
+#else
+	if (const UDataLayer* const* FoundDataLayer = NameToDataLayer.Find(InDataLayerName))
+	{
+		return *FoundDataLayer;
+	}
+#endif
+	return nullptr;
+}
+
+const UDataLayer* AWorldDataLayers::GetDataLayerFromLabel(const FName& InDataLayerLabel) const
+{
+#if WITH_EDITOR	
+	for (const UDataLayer* DataLayer : WorldDataLayers)
+	{
+		if (DataLayer->GetDataLayerLabel() == InDataLayerLabel)
+		{
+			return DataLayer;
+		}
+	}
+#else
+	if (const UDataLayer* const* FoundDataLayer = LabelToDataLayer.Find(InDataLayerLabel))
+	{
+		return *FoundDataLayer;
+	}
+#endif
+	return nullptr;
+}
+
+void AWorldDataLayers::ForEachDataLayer(TFunctionRef<bool(UDataLayer*)> Func)
+{
+	for (UDataLayer* DataLayer : WorldDataLayers)
+	{
+		if (!Func(DataLayer))
+		{
+			break;
+		}
+	}
+}
+
+void AWorldDataLayers::ForEachDataLayer(TFunctionRef<bool(UDataLayer*)> Func) const
+{
+	for (UDataLayer* DataLayer : WorldDataLayers)
+	{
+		if (!Func(DataLayer))
+		{
+			break;
+		}
+	}
+}
+
+void AWorldDataLayers::PostLoad()
+{
+	Super::PostLoad();
+
+#if !WITH_EDITOR
+	// Build acceleration tables
+	for (const UDataLayer* DataLayer : WorldDataLayers)
+	{
+		LabelToDataLayer.Add(DataLayer->GetDataLayerLabel(), DataLayer);
+		NameToDataLayer.Add(DataLayer->GetFName(), DataLayer);
+	}
+#endif
+}
+
+#undef LOCTEXT_NAMESPACE
\ No newline at end of file
diff --git a/Engine/Source/Runtime/Engine/Private/WorldPartition/WorldPartition.cpp b/Engine/Source/Runtime/Engine/Private/WorldPartition/WorldPartition.cpp
--- a/Engine/Source/Runtime/Engine/Private/WorldPartition/WorldPartition.cpp
+++ b/Engine/Source/Runtime/Engine/Private/WorldPartition/WorldPartition.cpp
@@ -22,8 +22,6 @@
 
 #if WITH_EDITOR
 #include "Editor.h"
-#include "Layers/Layer.h"
-#include "Layers/LayersSubsystem.h"
 #include "Editor/GroupActor.h"
 #include "EditorLevelUtils.h"
 #include "FileHelpers.h"
@@ -208,8 +206,6 @@ void UWorldPartition::Initialize(UWorld* InWorld, const FTransform& InTransform)
 
 		EditorHash->Initialize();
 
-		LayerSubSystem = GEditor->GetEditorSubsystem<ULayersSubsystem>();
-
 		if (IsMainWorldPartition())
 		{
 			RegisterDelegates();
@@ -225,8 +221,6 @@ void UWorldPartition::Initialize(UWorld* InWorld, const FTransform& InTransform)
 
 	if (bEditorOnly || !IsMainWorldPartition())
 	{
-		TSet<FName> AllLayersNames;
-
 		TArray<FAssetData> Assets;
 		UPackage* LevelPackage = OuterWorld->PersistentLevel->GetOutermost();
 
@@ -280,11 +274,6 @@ void UWorldPartition::Initialize(UWorld* InWorld, const FTransform& InTransform)
 				NewActorDesc->TransformInstance(ReplaceFrom, ReplaceTo, InstanceTransform);
 			}
 
-			if (bEditorOnly)
-			{
-				AllLayersNames.Append(NewActorDesc->GetLayers());
-			}
-
 			Actors.Add(NewActorDesc->GetGuid(), MoveTemp(NewActorDesc));
 		}
 
@@ -295,8 +284,6 @@ void UWorldPartition::Initialize(UWorld* InWorld, const FTransform& InTransform)
 				HashActorDesc(Pair.Value.Get());
 			}
 
-			CreateLayers(AllLayersNames);
-
 			// Load the always loaded cell, don't call LoadCells to avoid creating a transaction
 			UpdateLoadingEditorCell(EditorHash->GetAlwaysLoadedCell(), true);
 
@@ -723,20 +710,6 @@ void UWorldPartition::UpdateLoadingEditorCell(UWorldPartitionEditorCell* Cell, b
 	}
 }
 
-void UWorldPartition::CreateLayers(const TSet<FName>& LayerNames)
-{
-	// Make sure all the layers we discovered are created
-	if (LayerSubSystem)
-	{
-		for (FName LayerName : LayerNames)
-		{
-			if (!LayerSubSystem->IsLayer(LayerName))
-			{
-				LayerSubSystem->CreateLayer(LayerName);
-			}
-		}
-	}
-}
 #endif
 
 AActor* UWorldPartition::RegisterActor(FWorldPartitionActorDesc* ActorDesc)
@@ -757,16 +730,6 @@ AActor* UWorldPartition::RegisterActor(FWorldPartitionActorDesc* ActorDesc)
 
 	ApplyActorTransform(Actor, InstanceTransform);
 
-	// Since Actor might have been kept loaded but unregistered (modified), we need to make sure its visibility reflects its layer visibility
-	if (LayerSubSystem)
-	{
-		bool bActorModified = false;
-		bool bActorSelectionChanged = false;
-		const bool bActorNotifySelectionChange = false;
-		const bool bActorRedrawViewports = false;
-		LayerSubSystem->UpdateActorVisibility(Actor, bActorSelectionChanged, bActorModified, bActorNotifySelectionChange, bActorRedrawViewports);
-	}
-
 	TGuardValue<ITransaction*> TransGuard(GUndo, nullptr);
 	Actor->GetLevel()->AddLoadedActor(Actor);
 
@@ -788,11 +751,6 @@ void UWorldPartition::UnregisterActor(AActor* Actor)
 
 	if (!Actor->IsPendingKill())
 	{
-		if (LayerSubSystem)
-		{
-			LayerSubSystem->DisassociateActorFromLayers(Actor);
-		}
-
 		OnActorRegisteredEvent.Broadcast(*Actor, false);
 
 		Actor->GetLevel()->RemoveLoadedActor(Actor);
diff --git a/Engine/Source/Runtime/Engine/Private/WorldPartition/WorldPartitionActorDesc.cpp b/Engine/Source/Runtime/Engine/Private/WorldPartition/WorldPartitionActorDesc.cpp
--- a/Engine/Source/Runtime/Engine/Private/WorldPartition/WorldPartitionActorDesc.cpp
+++ b/Engine/Source/Runtime/Engine/Private/WorldPartition/WorldPartitionActorDesc.cpp
@@ -11,7 +11,9 @@
 #include "Engine/World.h"
 #include "Engine/Level.h"
 #include "Editor/GroupActor.h"
+#include "UObject/UE5MainStreamObjectVersion.h"
 #include "WorldPartition/WorldPartition.h"
+#include "WorldPartition/DataLayer/Datalayer.h"
 #endif
 
 #if WITH_EDITOR
@@ -53,7 +55,7 @@ void FWorldPartitionActorDesc::Init(const AActor* InActor)
 	RuntimeGrid = InActor->RuntimeGrid;
 	bActorIsEditorOnly = InActor->IsEditorOnly();
 	bLevelBoundsRelevant = InActor->IsLevelBoundsRelevant();
-	Layers = InActor->Layers;
+	DataLayers = InActor->GetDataLayerNames();
 	ActorPackage = InActor->GetPackage()->GetFName();
 	ActorPath = *InActor->GetPathName();
 	
@@ -139,12 +141,27 @@ FString FWorldPartitionActorDesc::ToString() const
 
 void FWorldPartitionActorDesc::Serialize(FArchive& Ar)
 {
-	Ar << Class << Guid << BoundsLocation << BoundsExtent << GridPlacement << RuntimeGrid << bActorIsEditorOnly << bLevelBoundsRelevant << Layers << References;
+	Ar.UsingCustomVersion(FUE5MainStreamObjectVersion::GUID);
+
+	Ar << Class << Guid << BoundsLocation << BoundsExtent << GridPlacement << RuntimeGrid << bActorIsEditorOnly << bLevelBoundsRelevant;
+	
+	if (Ar.CustomVer(FUE5MainStreamObjectVersion::GUID) < FUE5MainStreamObjectVersion::WorldPartitionActorDescSerializeDataLayers)
+	{
+		TArray<FName> Deprecated_Layers;
+		Ar << Deprecated_Layers;
+	}
+
+	Ar << References;
 
 	if (!Ar.IsPersistent())
 	{
 		Ar << ActorPackage << ActorPath;
 	}
+
+	if (Ar.CustomVer(FUE5MainStreamObjectVersion::GUID) >= FUE5MainStreamObjectVersion::WorldPartitionActorDescSerializeDataLayers)
+	{
+		Ar << DataLayers;
+	}
 }
 
 FBox FWorldPartitionActorDesc::GetBounds() const
diff --git a/Engine/Source/Runtime/Engine/Private/WorldPartition/WorldPartitionRuntimeCell.cpp b/Engine/Source/Runtime/Engine/Private/WorldPartition/WorldPartitionRuntimeCell.cpp
--- a/Engine/Source/Runtime/Engine/Private/WorldPartition/WorldPartitionRuntimeCell.cpp
+++ b/Engine/Source/Runtime/Engine/Private/WorldPartition/WorldPartitionRuntimeCell.cpp
@@ -15,6 +15,17 @@ void UWorldPartitionRuntimeCell::AddCellData(const UWorldPartitionRuntimeCellDat
 	CellDataMap.Add(InCellData->GetClass(), InCellData);
 }
 
+void UWorldPartitionRuntimeCell::SetDataLayers(const TArray<const UDataLayer*> InDataLayers)
+{
+	check(DataLayers.IsEmpty());
+	DataLayers.Reserve(InDataLayers.Num());
+	for (const UDataLayer* DataLayer : InDataLayers)
+	{
+		check(DataLayer->IsDynamicallyLoaded());
+		DataLayers.Add(DataLayer->GetFName());
+	}
+}
+
 #endif
 
 const UWorldPartitionRuntimeCellData* UWorldPartitionRuntimeCell::GetCellData(const TSubclassOf<UWorldPartitionRuntimeCellData> InCellDataClass) const
diff --git a/Engine/Source/Runtime/Engine/Private/WorldPartition/WorldPartitionRuntimeSpatialHash.cpp b/Engine/Source/Runtime/Engine/Private/WorldPartition/WorldPartitionRuntimeSpatialHash.cpp
--- a/Engine/Source/Runtime/Engine/Private/WorldPartition/WorldPartitionRuntimeSpatialHash.cpp
+++ b/Engine/Source/Runtime/Engine/Private/WorldPartition/WorldPartitionRuntimeSpatialHash.cpp
@@ -10,6 +10,8 @@
 #include "WorldPartition/WorldPartitionStreamingPolicy.h"
 #include "WorldPartition/DataLayer/DataLayerSubsystem.h"
 #include "WorldPartition/DataLayer/DataLayerHelper.h"
+#include "WorldPartition/DataLayer/WorldDataLayers.h"
+#include "WorldPartition/DataLayer/DataLayer.h"
 #include "GameFramework/WorldSettings.h"
 #include "ProfilingDebugging/ScopedTimers.h"
 #include "Engine/World.h"
@@ -110,24 +112,37 @@ static FAutoConsoleVariableRef CVarRuntimeSpatialHashCellToSourceAngleContributi
 // Clustering
 struct FActorCluster
 {
-	TSet<FGuid>			Actors;
-	EActorGridPlacement	GridPlacement;
-	FName				RuntimeGrid;
-	FBox				Bounds;
-	TArray<FName>		DataLayers;
-	uint32				DataLayersID;
-
-	FActorCluster(const FWorldPartitionActorDesc* ActorDesc);
+	TSet<FGuid>					Actors;
+	EActorGridPlacement			GridPlacement;
+	FName						RuntimeGrid;
+	FBox						Bounds;
+	TArray<const UDataLayer*>	DataLayers;
+	uint32						DataLayersID;
+
+	FActorCluster(const FWorldPartitionActorDesc* ActorDesc, UWorld* World);
 	void Add(const FActorCluster& ActorCluster);
 };
 
-FActorCluster::FActorCluster(const FWorldPartitionActorDesc* ActorDesc)
+FActorCluster::FActorCluster(const FWorldPartitionActorDesc* ActorDesc, UWorld* World)
 	: GridPlacement(ActorDesc->GetGridPlacement())
 	, RuntimeGrid(ActorDesc->GetRuntimeGrid())
 	, Bounds(ActorDesc->GetBounds())
 {
 	check(GridPlacement != EActorGridPlacement::None);
 	Actors.Add(ActorDesc->GetGuid());
+	if (const AWorldDataLayers* WorldDataLayers = AWorldDataLayers::Get(World))
+	{
+		for (const FName& DataLayerName : ActorDesc->GetDataLayers())
+		{
+			if (const UDataLayer* DataLayer = WorldDataLayers->GetDataLayerFromName(DataLayerName))
+			{
+				if (DataLayer->IsDynamicallyLoaded())
+				{
+					DataLayers.Add(DataLayer);
+				}
+			}
+		}
+	}
 	DataLayersID = FDataLayersHelper::ComputeDataLayerID(DataLayers);
 }
 
@@ -170,8 +185,9 @@ void FActorCluster::Add(const FActorCluster& ActorCluster)
 	// Merge DataLayers
 	if (DataLayersID != ActorCluster.DataLayersID)
 	{
-		for (const FName& DataLayer : ActorCluster.DataLayers)
+		for (const UDataLayer* DataLayer : ActorCluster.DataLayers)
 		{
+			check(DataLayer->IsDynamicallyLoaded());
 			DataLayers.AddUnique(DataLayer);
 		}
 		DataLayersID = FDataLayersHelper::ComputeDataLayerID(DataLayers);
@@ -181,10 +197,11 @@ void FActorCluster::Add(const FActorCluster& ActorCluster)
 
 void CreateActorCluster(const FGuid& ActorGuid, const FWorldPartitionActorDesc* ActorDesc, TMap<FGuid, FActorCluster*>& ActorToActorCluster, TSet<FActorCluster*>& ActorClustersSet, UWorldPartition* WorldPartition)
 {
+	UWorld* World = WorldPartition->GetWorld();
 	FActorCluster* ActorCluster = ActorToActorCluster.FindRef(ActorGuid);
 	if (!ActorCluster)
 	{
-		ActorCluster = new FActorCluster(ActorDesc);
+		ActorCluster = new FActorCluster(ActorDesc, World);
 		ActorClustersSet.Add(ActorCluster);
 		ActorToActorCluster.Add(ActorGuid, ActorCluster);
 	}
@@ -217,7 +234,7 @@ void CreateActorCluster(const FGuid& ActorGuid, const FWorldPartitionActorDesc*
 				else
 				{
 					// Put Reference in Actor's cluster
-					ActorCluster->Add(FActorCluster(ReferenceActorDesc));
+					ActorCluster->Add(FActorCluster(ReferenceActorDesc, World));
 				}
 
 				// Map its cluster
@@ -469,28 +486,30 @@ struct FSquare2DGridHelper
 
 	struct FGridLevel : public FGrid2D
 	{
+#if WITH_EDITOR
 		struct FGridCellDataChunk
 		{
-			FGridCellDataChunk(const TArray<FName>& InDataLayers)
-			: DataLayers(InDataLayers)
-			, DataLayersID(FDataLayersHelper::ComputeDataLayerID(InDataLayers))
-			{}
+			FGridCellDataChunk(const TArray<const UDataLayer*>& InDataLayers)
+			{
+				Algo::TransformIf(InDataLayers, DataLayers, [](const UDataLayer* DataLayer) { return DataLayer->IsDynamicallyLoaded(); }, [](const UDataLayer* DataLayer) { return DataLayer; });
+				DataLayersID = FDataLayersHelper::ComputeDataLayerID(DataLayers);
+			}
 
 			void AddActor(const FGuid& Actor) { Actors.Add(Actor); }
 			const TSet<FGuid>& GetActors() const { return Actors; }
 			bool HasDataLayers() const { return !DataLayers.IsEmpty(); }
-			const TArray<FName>& GetDataLayers() const { return DataLayers; }
+			const TArray<const UDataLayer*>& GetDataLayers() const { return DataLayers; }
 			uint32 GetDataLayersID() const { return DataLayersID; }
 
 		private:
 			TSet<FGuid> Actors;
-			TArray<FName> DataLayers;
+			TArray<const UDataLayer*> DataLayers;
 			uint32 DataLayersID;
 		};
 
 		struct FGridCell
 		{
-			void AddActor(const FGuid& InActor, const TArray<FName>& InDataLayers)
+			void AddActor(const FGuid& InActor, const TArray<const UDataLayer*>& InDataLayers)
 			{
 				uint32 DataLayersID = FDataLayersHelper::ComputeDataLayerID(InDataLayers);
 				FGridCellDataChunk* ActorDataChunk = Algo::FindByPredicate(DataChunks, [&](FGridCellDataChunk& InDataChunk) { return InDataChunk.GetDataLayersID() == DataLayersID; });
@@ -501,7 +520,7 @@ struct FSquare2DGridHelper
 				ActorDataChunk->AddActor(InActor);
 			}
 
-			void AddActors(const TSet<FGuid>& InActors, const TArray<FName>& InDataLayers)
+			void AddActors(const TSet<FGuid>& InActors, const TArray<const UDataLayer*>& InDataLayers)
 			{
 				for (const FGuid& Actor : InActors)
 				{
@@ -526,13 +545,17 @@ struct FSquare2DGridHelper
 		private:
 			TArray<FGridCellDataChunk> DataChunks;
 		};
+#endif
 
 		inline FGridLevel(const FVector2D& InOrigin, int32 InCellSize, int32 InGridSize)
 			: FGrid2D(InOrigin, InCellSize, InGridSize)
 		{
+#if WITH_EDITOR
 			Cells.InsertDefaulted(0, GridSize * GridSize);
+#endif
 		}
 
+#if WITH_EDITOR
 		/**
 		 * Returns the cell at the specified grid coordinate
 		 *
@@ -560,6 +583,7 @@ struct FSquare2DGridHelper
 		}
 
 		TArray<FGridCell> Cells;
+#endif
 	};
 
 	FSquare2DGridHelper(int32 InNumLevels, const FVector& InOrigin, int32 InCellSize, int32 InGridSize)
@@ -588,6 +612,7 @@ struct FSquare2DGridHelper
 		}
 	}
 
+#if WITH_EDITOR
 	// Returns the lowest grid level
 	inline FGridLevel& GetLowestLevel() { return Levels[0]; }
 
@@ -599,6 +624,7 @@ struct FSquare2DGridHelper
 
 	// Returns the cell at the given coord
 	inline const FGridLevel::FGridCell& GetCell(const FIntVector& InCoords) const { return Levels[InCoords.Z].GetCell(FIntVector2(InCoords.X, InCoords.Y)); }
+#endif
 
 	/**
 	 * Returns the cell bounds
@@ -1023,7 +1049,18 @@ static FSquare2DGridHelper GetPartitionedActors(const UWorldPartition* WorldPart
 				const FWorldPartitionActorDesc& ActorDesc = *WorldPartition->GetActorDesc(ActorGuid);
 				if (PartitionedActors.GetLowestLevel().GetCellCoords(FVector2D(ActorDesc.GetOrigin()), CellCoords))
 				{
-					PartitionedActors.GetLowestLevel().GetCell(CellCoords).AddActor(ActorGuid, TArray<FName>());
+					TArray<const UDataLayer*> ActorDataLayers;
+					if (const AWorldDataLayers* WorldDataLayers = AWorldDataLayers::Get(WorldPartition->GetWorld()))
+					{
+						for (const FName& DataLayerName : ActorDesc.GetDataLayers())
+						{
+							if (const UDataLayer* DataLayer = WorldDataLayers->GetDataLayerFromName(DataLayerName))
+							{
+								ActorDataLayers.Add(DataLayer);
+							}
+						}
+					}
+					PartitionedActors.GetLowestLevel().GetCell(CellCoords).AddActor(ActorGuid, ActorDataLayers);
 				}
 				else
 				{
diff --git a/Engine/Source/Runtime/Engine/Public/WorldPartition/DataLayer/ActorDataLayer.h b/Engine/Source/Runtime/Engine/Public/WorldPartition/DataLayer/ActorDataLayer.h
new file mode 100644
--- /dev/null
+++ b/Engine/Source/Runtime/Engine/Public/WorldPartition/DataLayer/ActorDataLayer.h
@@ -0,0 +1,27 @@
+// Copyright Epic Games, Inc. All Rights Reserved.
+
+#pragma once
+
+#include "CoreMinimal.h"
+#include "ActorDataLayer.generated.h"
+
+USTRUCT(BlueprintType)
+struct ENGINE_API FActorDataLayer
+{
+	GENERATED_USTRUCT_BODY()
+
+	FActorDataLayer()
+	: Name(NAME_None)
+	{}
+
+	FActorDataLayer(const FName& InName)
+	: Name(InName)
+	{}
+
+	FORCEINLINE bool operator==(const FActorDataLayer& Other) const { return Name == Other.Name; }
+	FORCEINLINE bool operator<(const FActorDataLayer& Other) const { return Name.FastLess(Other.Name); }
+
+	/** The name of this layer */
+	UPROPERTY(EditAnywhere, BlueprintReadWrite, Category = DataLayer)
+	FName Name;
+};
\ No newline at end of file
diff --git a/Engine/Source/Runtime/Engine/Public/WorldPartition/DataLayer/DataLayer.h b/Engine/Source/Runtime/Engine/Public/WorldPartition/DataLayer/DataLayer.h
new file mode 100644
--- /dev/null
+++ b/Engine/Source/Runtime/Engine/Public/WorldPartition/DataLayer/DataLayer.h
@@ -0,0 +1,41 @@
+// Copyright Epic Games, Inc. All Rights Reserved.
+
+#pragma once
+
+#include "CoreMinimal.h"
+#include "UObject/ObjectMacros.h"
+#include "UObject/UObjectGlobals.h"
+#include "UObject/Object.h"
+#include "WorldDataLayers.h"
+
+#include "DataLayer.generated.h"
+
+UCLASS(Config = Engine, PerObjectConfig, Within = WorldDataLayers)
+class ENGINE_API UDataLayer : public UObject
+{
+	GENERATED_UCLASS_BODY()
+
+public:
+#if WITH_EDITOR
+	void SetDataLayerLabel(FName InDataLayerLabel);
+	void SetVisible(bool bInIsVisible);
+	void SetIsDynamicallyLoaded(bool bInIsDynamicallyLoaded);
+#endif
+	
+	FName GetDataLayerLabel() const  { return DataLayerLabel; }
+	bool IsVisible() const { return bIsVisible; }
+	bool IsDynamicallyLoaded() const { return bIsDynamicallyLoaded; }
+
+private:
+	/** The display name of the DataLayer */
+	UPROPERTY()
+	FName DataLayerLabel;
+
+	/** Whether actors associated with the DataLayer are visible in the viewport */
+	UPROPERTY()
+	uint32 bIsVisible : 1;
+
+	/** Whether the layer can affect actor loading or not when layer is activated */
+	UPROPERTY()
+	uint32 bIsDynamicallyLoaded : 1;
+};
\ No newline at end of file
diff --git a/Engine/Source/Runtime/Engine/Public/WorldPartition/DataLayer/DataLayerHelper.h b/Engine/Source/Runtime/Engine/Public/WorldPartition/DataLayer/DataLayerHelper.h
--- a/Engine/Source/Runtime/Engine/Public/WorldPartition/DataLayer/DataLayerHelper.h
+++ b/Engine/Source/Runtime/Engine/Public/WorldPartition/DataLayer/DataLayerHelper.h
@@ -4,8 +4,13 @@
 
 #include "CoreMinimal.h"
 
+class UDataLayer;
+
 struct ENGINE_API FDataLayersHelper
 {
+#if WITH_EDITOR
 	static const uint32 NoDataLayerID = 0;
-	static uint32 ComputeDataLayerID(const TArray<FName>& DataLayers);
+
+	static uint32 ComputeDataLayerID(const TArray<const UDataLayer*>& DataLayers);
+#endif
 };
\ No newline at end of file
diff --git a/Engine/Source/Runtime/Engine/Public/WorldPartition/DataLayer/DataLayerSubsystem.h b/Engine/Source/Runtime/Engine/Public/WorldPartition/DataLayer/DataLayerSubsystem.h
--- a/Engine/Source/Runtime/Engine/Public/WorldPartition/DataLayer/DataLayerSubsystem.h
+++ b/Engine/Source/Runtime/Engine/Public/WorldPartition/DataLayer/DataLayerSubsystem.h
@@ -3,6 +3,7 @@
 #pragma once
 
 #include "Subsystems/WorldSubsystem.h"
+#include "WorldPartition/DataLayer/DataLayer.h"
 #include "DataLayerSubsystem.generated.h"
 
 /**
@@ -21,14 +22,15 @@ public:
 	virtual bool ShouldCreateSubsystem(UObject* Outer) const override;
 	//~ End USubsystem Interface.
 
-	void ActivateDataLayer(const FName& InDataLayer, bool bActivate);
-	bool IsDataLayerActive(const FName& InDataLayer) const;
-	bool IsAnyDataLayerActive(const TArray<FName>& InDataLayers) const;
-	const TSet<FName>& GetActiveDataLayers() const { return ActiveDataLayers; }
+	const UDataLayer* GetDataLayerFromLabel(const FName& InDataLayerLabel) const;
+	const UDataLayer* GetDataLayerFromName(const FName& InDataLayerName) const;
+	void ActivateDataLayer(const FName& InDataLayerName, bool bActivate);
+	bool IsDataLayerActive(const FName& InDataLayerName) const;
+	bool IsAnyDataLayerActive(const TArray<FName>& InDataLayerNames) const;
 
 private:
 	/** Console command used to toggle activation of a DataLayer */
 	static class FAutoConsoleCommand ToggleDataLayerActivation;
 
-	TSet<FName> ActiveDataLayers;
+	TSet<FName> ActiveDataLayerNames;
 };
\ No newline at end of file
diff --git a/Engine/Source/Runtime/Engine/Public/WorldPartition/DataLayer/WorldDataLayers.h b/Engine/Source/Runtime/Engine/Public/WorldPartition/DataLayer/WorldDataLayers.h
new file mode 100644
--- /dev/null
+++ b/Engine/Source/Runtime/Engine/Public/WorldPartition/DataLayer/WorldDataLayers.h
@@ -0,0 +1,47 @@
+// Copyright Epic Games, Inc. All Rights Reserved.
+
+#pragma once
+
+#include "CoreMinimal.h"
+#include "UObject/ObjectMacros.h"
+#include "GameFramework/Info.h"
+#include "WorldDataLayers.generated.h"
+
+class UDataLayer;
+
+/**
+ * Actor containing all data layers for a world
+ */
+UCLASS(hidecategories = (Actor, Advanced, Display, Events, Object, Attachment, Info, Input, Blueprint, Layers, Tags, Replication), notplaceable)
+class ENGINE_API AWorldDataLayers : public AInfo
+{
+	GENERATED_UCLASS_BODY()
+
+public:
+	static const AWorldDataLayers* Get(UWorld* World);
+	virtual void PostLoad() override;
+
+#if WITH_EDITOR
+	virtual EActorGridPlacement GetDefaultGridPlacement() const override { return EActorGridPlacement::AlwaysLoaded; }
+
+	UDataLayer* CreateDataLayer();
+	bool RemoveDataLayer(const UDataLayer* InDataLayer);
+	bool RemoveDataLayers(const TArray<const UDataLayer*>& InDataLayers);
+	FName GenerateUniqueDataLayerLabel(const FName& InDataLayerLabel) const;
+#endif
+	
+	bool ContainsDataLayer(const UDataLayer* InDataLayer) const;
+	const UDataLayer* GetDataLayerFromName(const FName& InDataLayerName) const;
+	const UDataLayer* GetDataLayerFromLabel(const FName& InDataLayerLabel) const;
+	void ForEachDataLayer(TFunctionRef<bool(class UDataLayer*)> Func);
+	void ForEachDataLayer(TFunctionRef<bool(class UDataLayer*)> Func) const;
+
+private:
+#if !WITH_EDITOR
+	TMap<FName, const UDataLayer*> LabelToDataLayer;
+	TMap<FName, const UDataLayer*> NameToDataLayer;
+#endif
+
+	UPROPERTY()
+	TSet<UDataLayer*> WorldDataLayers;
+};
\ No newline at end of file
diff --git a/Engine/Source/Runtime/Engine/Public/WorldPartition/WorldPartition.h b/Engine/Source/Runtime/Engine/Public/WorldPartition/WorldPartition.h
--- a/Engine/Source/Runtime/Engine/Public/WorldPartition/WorldPartition.h
+++ b/Engine/Source/Runtime/Engine/Public/WorldPartition/WorldPartition.h
@@ -160,7 +160,6 @@ public:
 
 #if WITH_EDITOR
 	TMap<FGuid, TUniquePtr<FWorldPartitionActorDesc>> Actors;
-	class ULayersSubsystem* LayerSubSystem;
 	
 	static TUniquePtr<FWorldPartitionActorDescFactory> DefaultActorDescFactory;
 	static TMap<FName, FWorldPartitionActorDescFactory*> ActorDescFactories;
@@ -205,8 +204,6 @@ private:
 
 #if WITH_EDITOR
 	void UpdateLoadingEditorCell(UWorldPartitionEditorCell* Cell, bool bShouldBeLoaded);
-	void CreateLayers(const TSet<FName>& LayerNames);
-
 	void HashActorDesc(FWorldPartitionActorDesc* ActorDesc);
 	void UnhashActorDesc(FWorldPartitionActorDesc* ActorDesc);
 #endif
diff --git a/Engine/Source/Runtime/Engine/Public/WorldPartition/WorldPartitionActorDesc.h b/Engine/Source/Runtime/Engine/Public/WorldPartition/WorldPartitionActorDesc.h
--- a/Engine/Source/Runtime/Engine/Public/WorldPartition/WorldPartitionActorDesc.h
+++ b/Engine/Source/Runtime/Engine/Public/WorldPartition/WorldPartitionActorDesc.h
@@ -41,7 +41,7 @@ public:
 	inline FName GetRuntimeGrid() const { return RuntimeGrid; }
 	inline bool GetActorIsEditorOnly() const { return bActorIsEditorOnly; }
 	inline bool GetLevelBoundsRelevant() const { return bLevelBoundsRelevant; }
-	inline const TArray<FName>& GetLayers() const { return Layers; }
+	inline const TArray<FName>& GetDataLayers() const { return DataLayers; }
 	inline FName GetActorPackage() const { return ActorPackage; }
 	inline FName GetActorPath() const { return ActorPath; }
 	FBox GetBounds() const;
@@ -95,7 +95,7 @@ protected:
 	FName						RuntimeGrid;
 	bool						bActorIsEditorOnly;
 	bool						bLevelBoundsRelevant;
-	TArray<FName>				Layers;
+	TArray<FName>				DataLayers;
 	TArray<FGuid>				References;
 	
 	mutable uint32				LoadedRefCount;
diff --git a/Engine/Source/Runtime/Engine/Public/WorldPartition/WorldPartitionRuntimeCell.h b/Engine/Source/Runtime/Engine/Public/WorldPartition/WorldPartitionRuntimeCell.h
--- a/Engine/Source/Runtime/Engine/Public/WorldPartition/WorldPartitionRuntimeCell.h
+++ b/Engine/Source/Runtime/Engine/Public/WorldPartition/WorldPartitionRuntimeCell.h
@@ -3,6 +3,7 @@
 
 #include "CoreMinimal.h"
 #include "Templates/SubclassOf.h"
+#include "WorldPartition/DataLayer/DataLayer.h"
 #include "WorldPartitionRuntimeCell.generated.h"
 
 USTRUCT()
@@ -53,11 +54,11 @@ class UWorldPartitionRuntimeCell : public UObject
 	virtual FLinearColor GetDebugColor() const PURE_VIRTUAL(UWorldPartitionRuntimeCell::GetDebugColor, return FLinearColor::Black;);
 	virtual bool IsAlwaysLoaded() const { return bIsAlwaysLoaded; }
 	virtual void SetIsAlwaysLoaded(bool bInIsAlwaysLoaded) { bIsAlwaysLoaded = bInIsAlwaysLoaded; }
-	void SetDataLayers(const TArray<FName> InDataLayers) { check(DataLayers.IsEmpty()); DataLayers = InDataLayers; }
 	bool HasDataLayers() const { return !DataLayers.IsEmpty(); }
 	const TArray<FName>& GetDataLayers() const { return DataLayers; }
 
 #if WITH_EDITOR
+	void SetDataLayers(const TArray<const UDataLayer*> InDataLayers);
 	virtual void AddActorToCell(FName Package, FName Path) PURE_VIRTUAL(UWorldPartitionRuntimeCell::AddActorToCell,);
 	virtual bool CreateCellForCook() PURE_VIRTUAL(UWorldPartitionRuntimeCell::CreateCellForCook, return false;);
 	virtual int32 GetActorCount() const PURE_VIRTUAL(UWorldPartitionRuntimeCell::GetActorCount, return 0;);
