Fixed Linux build warnings.

diff --git a/Engine/Source/Editor/DataLayerEditor/Private/DataLayer/DataLayerHierarchy.cpp b/Engine/Source/Editor/DataLayerEditor/Private/DataLayer/DataLayerHierarchy.cpp
--- a/Engine/Source/Editor/DataLayerEditor/Private/DataLayer/DataLayerHierarchy.cpp
+++ b/Engine/Source/Editor/DataLayerEditor/Private/DataLayer/DataLayerHierarchy.cpp
@@ -161,7 +161,7 @@ void FDataLayerHierarchy::OnActorDataLayersChanged(const TWeakObjectPtr<AActor>&
 void FDataLayerHierarchy::OnDataLayerChanged(const EDataLayerAction Action, const TWeakObjectPtr<const UDataLayer>& ChangedDataLayer, const FName& ChangedProperty)
 {
 	const UDataLayer* DataLayer = ChangedDataLayer.Get();
-	if (DataLayer && (RepresentingWorld.Get() == DataLayer->GetWorld()) || Action == EDataLayerAction::Delete || Action == EDataLayerAction::Reset)
+	if ((DataLayer && (RepresentingWorld.Get() == DataLayer->GetWorld())) || (Action == EDataLayerAction::Delete) || (Action == EDataLayerAction::Reset))
 	{
 		FullRefreshEvent();
 	}
diff --git a/Engine/Source/Runtime/Engine/Private/WorldPartition/WorldPartitionActorDesc.cpp b/Engine/Source/Runtime/Engine/Private/WorldPartition/WorldPartitionActorDesc.cpp
--- a/Engine/Source/Runtime/Engine/Private/WorldPartition/WorldPartitionActorDesc.cpp
+++ b/Engine/Source/Runtime/Engine/Private/WorldPartition/WorldPartitionActorDesc.cpp
@@ -13,7 +13,7 @@
 #include "Editor/GroupActor.h"
 #include "UObject/UE5MainStreamObjectVersion.h"
 #include "WorldPartition/WorldPartition.h"
-#include "WorldPartition/DataLayer/Datalayer.h"
+#include "WorldPartition/DataLayer/DataLayer.h"
 #endif
 
 #if WITH_EDITOR
