Fixed non-unity build.

diff --git a/Engine/Source/Editor/DataLayerEditor/Private/DataLayer/DataLayerMode.cpp b/Engine/Source/Editor/DataLayerEditor/Private/DataLayer/DataLayerMode.cpp
--- a/Engine/Source/Editor/DataLayerEditor/Private/DataLayer/DataLayerMode.cpp
+++ b/Engine/Source/Editor/DataLayerEditor/Private/DataLayer/DataLayerMode.cpp
@@ -16,6 +16,7 @@
 #include "EditorActorFolders.h"
 #include "Algo/Transform.h"
 #include "ToolMenus.h"
+#include "Selection.h"
 #include "Editor.h"
 
 #define LOCTEXT_NAMESPACE "DataLayer"
