Fixed UDataLayerModeSettings not saving in EditorPerProjectUserSettings.ini.

diff --git a/Engine/Source/Editor/DataLayerEditor/Private/DataLayer/SDataLayerBrowser.cpp b/Engine/Source/Editor/DataLayerEditor/Private/DataLayer/SDataLayerBrowser.cpp
--- a/Engine/Source/Editor/DataLayerEditor/Private/DataLayer/SDataLayerBrowser.cpp
+++ b/Engine/Source/Editor/DataLayerEditor/Private/DataLayer/SDataLayerBrowser.cpp
@@ -19,7 +19,7 @@
 
 void SDataLayerBrowser::Construct(const FArguments& InArgs)
 {
-	const UDataLayerModeSettings* SharedSettings = GetMutableDefault<UDataLayerModeSettings>();
+	const UDataLayerModeSettings* SharedSettings = GetDefault<UDataLayerModeSettings>();
 	Mode = SharedSettings->bShowDataLayerContent ? EDataLayerBrowserMode::DataLayerContents : EDataLayerBrowserMode::DataLayers;
 
 	auto ToggleDataLayerContents = [this]()
@@ -151,7 +151,7 @@ void SDataLayerBrowser::SetupDataLayerMode(EDataLayerBrowserMode InNewMode)
 	
 	UDataLayerModeSettings* SharedSettings = GetMutableDefault<UDataLayerModeSettings>();
 	SharedSettings->bShowDataLayerContent = Mode == EDataLayerBrowserMode::DataLayerContents;
-	SharedSettings->PostEditChange();
+	SharedSettings->SaveConfig();
 
 	ModeChanged.Broadcast(Mode);
 }
