Removed redundant call when renaming DataLayer label in browser.

diff --git a/Engine/Source/Editor/DataLayerEditor/Private/DataLayer/SDataLayerTreeLabel.cpp b/Engine/Source/Editor/DataLayerEditor/Private/DataLayer/SDataLayerTreeLabel.cpp
--- a/Engine/Source/Editor/DataLayerEditor/Private/DataLayer/SDataLayerTreeLabel.cpp
+++ b/Engine/Source/Editor/DataLayerEditor/Private/DataLayer/SDataLayerTreeLabel.cpp
@@ -172,7 +172,6 @@ void SDataLayerTreeLabel::OnLabelCommitted(const FText& InLabel, ETextCommit::Ty
 	{
 		const FScopedTransaction Transaction(LOCTEXT("SceneOutlinerRenameDataLayerTransaction", "Rename Data Layer"));
 		UDataLayerEditorSubsystem::Get()->RenameDataLayer(DataLayer, *InLabel.ToString());
-		DataLayer->SetDataLayerLabel(*InLabel.ToString());
 
 		if (WeakSceneOutliner.IsValid())
 		{
