- Support drag and drop of actors into selection of DataLayers in DataLayer browser
- Keep selection when refreshing DataLayer browser

diff --git a/Engine/Source/Editor/DataLayerEditor/Private/DataLayer/DataLayerMode.cpp b/Engine/Source/Editor/DataLayerEditor/Private/DataLayer/DataLayerMode.cpp
--- a/Engine/Source/Editor/DataLayerEditor/Private/DataLayer/DataLayerMode.cpp
+++ b/Engine/Source/Editor/DataLayerEditor/Private/DataLayer/DataLayerMode.cpp
@@ -71,7 +71,17 @@ void FDataLayerMode::OnItemAdded(FSceneOutlinerTreeItemPtr Item)
 	{
 		if (!Item->Flags.bIsFilteredOut)
 		{
-			SceneOutliner->SetItemSelection(Item, true);
+			if (SelectedDataLayers.Contains(DataLayerItem->GetDataLayer()))
+			{
+				SceneOutliner->AddToSelection({Item});
+			}
+		}
+	}
+	else if (const FDataLayerActorTreeItem* DataLayerActorTreeItem = Item->CastTo<FDataLayerActorTreeItem>())
+	{
+		if (SelectedDataLayerActors.Contains(FSelectedDataLayerActor(DataLayerActorTreeItem->GetDataLayer(), DataLayerActorTreeItem->GetActor())))
+		{
+			SceneOutliner->AddToSelection({Item});
 		}
 	}
 }
@@ -194,6 +204,10 @@ FSceneOutlinerDragValidationInfo FDataLayerMode::ValidateDrop(const ISceneOutlin
 			{
 				return FSceneOutlinerDragValidationInfo(ESceneOutlinerDropCompatibility::IncompatibleGeneric, FText());
 			}
+			if (SceneOutliner->GetTree().IsItemSelected(const_cast<ISceneOutlinerTreeItem&>(DropTarget).AsShared()) && (GetSelectedDataLayers(SceneOutliner).Num() > 1))
+			{
+				return FSceneOutlinerDragValidationInfo(ESceneOutlinerDropCompatibility::Compatible, LOCTEXT("AssignToDataLayers", "Assign to Selected Data Layers"));
+			}
 			return FSceneOutlinerDragValidationInfo(ESceneOutlinerDropCompatibility::Compatible, FText::Format(LOCTEXT("AssignToDataLayer", "Assign to Data Layer \"{0}\""), FText::FromName(DataLayerTarget->GetDataLayerLabel())));
 		}
 		else
@@ -271,8 +285,8 @@ TArray<AActor*> FDataLayerMode::GetActorsFromOperation(const FDragDropOperation&
 
 void FDataLayerMode::OnDrop(ISceneOutlinerTreeItem& DropTarget, const FSceneOutlinerDragDropPayload& Payload, const FSceneOutlinerDragValidationInfo& ValidationInfo) const
 {
-	TArray<AActor*> ActorsToAddToDataLayer = GetActorsFromOperation(Payload.SourceOperation);
-	if (ActorsToAddToDataLayer.IsEmpty())
+	TArray<AActor*> ActorsToAdd = GetActorsFromOperation(Payload.SourceOperation);
+	if (ActorsToAdd.IsEmpty())
 	{
 		return;
 	}
@@ -281,19 +295,30 @@ void FDataLayerMode::OnDrop(ISceneOutlinerTreeItem& DropTarget, const FSceneOutl
 	{
 		if (const UDataLayer* DataLayer = DataLayerItem->GetDataLayer())
 		{
-			const FScopedTransaction Transaction(LOCTEXT("DataLayerOutlinerAddActorsToDataLayer", "Add Actors to DataLayers"));
-			DataLayerEditorSubsystem->AddActorsToDataLayer(ActorsToAddToDataLayer, DataLayer);
+			if (SceneOutliner->GetTree().IsItemSelected(const_cast<ISceneOutlinerTreeItem&>(DropTarget).AsShared()))
+			{
+				TArray<const UDataLayer*> AllSelectedDataLayers = GetSelectedDataLayers(SceneOutliner);
+				if (AllSelectedDataLayers.Num() > 1)
+				{
+					const FScopedTransaction Transaction(LOCTEXT("DataLayerOutlinerAddActorsToDataLayers", "Add Actors to DataLayers"));
+					DataLayerEditorSubsystem->AddActorsToDataLayers(ActorsToAdd, AllSelectedDataLayers);
+					return;
+				}
+			}
+
+			const FScopedTransaction Transaction(LOCTEXT("DataLayerOutlinerAddActorsToDataLayer", "Add Actors to DataLayer"));
+			DataLayerEditorSubsystem->AddActorsToDataLayer(ActorsToAdd, DataLayer);
 		}
 	}
 	else
 	{
-		if (!ActorsToAddToDataLayer[0]->HasDataLayers())
+		if (!ActorsToAdd[0]->HasDataLayers())
 		{
 			// Only allow actors not coming from the DataLayerBrowser
 			const FScopedTransaction Transaction(LOCTEXT("AddSelectedActorsToNewDataLayer", "Add Actors to New DataLayer"));
 			if (const UDataLayer* NewDataLayer = DataLayerEditorSubsystem->CreateDataLayer())
 			{
-				DataLayerEditorSubsystem->AddActorsToDataLayer(ActorsToAddToDataLayer, NewDataLayer);
+				DataLayerEditorSubsystem->AddActorsToDataLayer(ActorsToAdd, NewDataLayer);
 			}
 		}
 	}
@@ -372,6 +397,16 @@ TSharedPtr<FDragDropOperation> FDataLayerMode::CreateDragDropOperation(const TAr
 static const FName DefaultContextBaseMenuName("DataLayerOutliner.DefaultContextMenuBase");
 static const FName DefaultContextMenuName("DataLayerOutliner.DefaultContextMenu");
 
+TArray<const UDataLayer*> FDataLayerMode::GetSelectedDataLayers(SSceneOutliner* InSceneOutliner) const
+{
+	FSceneOutlinerItemSelection ItemSelection(InSceneOutliner->GetSelection());
+	TArray<FDataLayerTreeItem*> SelectedDataLayerItems;
+	ItemSelection.Get<FDataLayerTreeItem>(SelectedDataLayerItems);
+	TArray<const UDataLayer*> ValidSelectedDataLayers;
+	Algo::TransformIf(SelectedDataLayerItems, ValidSelectedDataLayers, [](const auto Item) { return Item && Item->GetDataLayer(); }, [](const auto Item) { return Item->GetDataLayer(); });
+	return MoveTemp(ValidSelectedDataLayers);
+}
+
 void FDataLayerMode::RegisterContextMenu()
 {
 	UToolMenus* ToolMenus = UToolMenus::Get();
@@ -389,20 +424,16 @@ void FDataLayerMode::RegisterContextMenu()
 			}
 
 			SSceneOutliner* SceneOutliner = Context->SceneOutliner.Pin().Get();
-			FSceneOutlinerItemSelection ItemSelection(SceneOutliner->GetSelection());
-			TArray<FDataLayerTreeItem*> SelectedDataLayerItems;
-			ItemSelection.Get<FDataLayerTreeItem>(SelectedDataLayerItems);
-			TArray<const UDataLayer*> SelectedDataLayers;
-			Algo::TransformIf(SelectedDataLayerItems, SelectedDataLayers, [](const auto Item) { return Item && Item->GetDataLayer(); }, [](const auto Item) { return Item->GetDataLayer(); });
+			TArray<const UDataLayer*> SelectedDataLayers = GetSelectedDataLayers(SceneOutliner);
 
 			TArray<const UDataLayer*> AllDataLayers;
 			if (const AWorldDataLayers* WorldDataLayers = AWorldDataLayers::Get(RepresentingWorld.Get()))
 			{
 				WorldDataLayers->ForEachDataLayer([&AllDataLayers](UDataLayer* DataLayer)
-					{
-						AllDataLayers.Add(DataLayer);
-						return true;
-					});
+				{
+					AllDataLayers.Add(DataLayer);
+					return true;
+				});
 			}
 
 			{
@@ -577,10 +608,18 @@ TUniquePtr<ISceneOutlinerHierarchy> FDataLayerMode::CreateHierarchy()
 	return ActorHierarchy;
 }
 
+void FDataLayerMode::OnItemSelectionChanged(FSceneOutlinerTreeItemPtr TreeItem, ESelectInfo::Type SelectionType, const FSceneOutlinerItemSelection& Selection)
+{
+	SelectedDataLayers.Empty();
+	SelectedDataLayerActors.Empty();
+	Selection.ForEachItem<FDataLayerTreeItem>([this](const FDataLayerTreeItem& Item) { SelectedDataLayers.Add(Item.GetDataLayer()); });
+	Selection.ForEachItem<FDataLayerActorTreeItem>([this](const FDataLayerActorTreeItem& Item) { SelectedDataLayerActors.Add(FSelectedDataLayerActor(Item.GetDataLayer(), Item.GetActor())); });
+}
+
 void FDataLayerMode::Rebuild()
 {
 	ChooseRepresentingWorld();
-
+	
 	if (Hierarchy)
 	{
 		Hierarchy->OnHierarchyChanged().Clear();
diff --git a/Engine/Source/Editor/DataLayerEditor/Private/DataLayer/DataLayerMode.h b/Engine/Source/Editor/DataLayerEditor/Private/DataLayer/DataLayerMode.h
--- a/Engine/Source/Editor/DataLayerEditor/Private/DataLayer/DataLayerMode.h
+++ b/Engine/Source/Editor/DataLayerEditor/Private/DataLayer/DataLayerMode.h
@@ -48,6 +48,7 @@ public:
 	virtual void OnItemDoubleClick(FSceneOutlinerTreeItemPtr Item) override;
 	virtual void OnItemAdded(FSceneOutlinerTreeItemPtr Item) override;
 	virtual FReply OnKeyDown(const FKeyEvent& InKeyEvent) override;
+	virtual void OnItemSelectionChanged(FSceneOutlinerTreeItemPtr TreeItem, ESelectInfo::Type SelectionType, const FSceneOutlinerItemSelection& Selection) override;
 
 	virtual bool CanSupportDragAndDrop() const { return true; }
 	virtual bool ParseDragDrop(FSceneOutlinerDragDropPayload& OutPayload, const FDragDropOperation& Operation) const override;
@@ -66,6 +67,7 @@ private:
 	void RegisterContextMenu();
 	void ChooseRepresentingWorld();
 	TArray<AActor*> GetActorsFromOperation(const FDragDropOperation& Operation, bool bOnlyFindFirst = false) const;
+	TArray<const UDataLayer*> GetSelectedDataLayers(SSceneOutliner* InSceneOutliner) const;
 
 	/** The world which we are currently representing */
 	TWeakObjectPtr<UWorld> RepresentingWorld;
@@ -77,4 +79,8 @@ private:
 	UDataLayerEditorSubsystem* DataLayerEditorSubsystem;
 	/** If this mode was created to display a specific world, don't allow it to be reassigned */
 	const TWeakObjectPtr<UWorld> SpecifiedWorldToDisplay;
+
+	TSet<TWeakObjectPtr<const UDataLayer>> SelectedDataLayers;
+	typedef TPair<TWeakObjectPtr<const UDataLayer>, TWeakObjectPtr<const AActor>> FSelectedDataLayerActor;
+	TSet<FSelectedDataLayerActor> SelectedDataLayerActors;
 };
