Fix DataLayer build break

#rnx

diff --git a/Engine/Source/Editor/DataLayerEditor/Private/DataLayer/DataLayerMode.cpp b/Engine/Source/Editor/DataLayerEditor/Private/DataLayer/DataLayerMode.cpp
--- a/Engine/Source/Editor/DataLayerEditor/Private/DataLayer/DataLayerMode.cpp
+++ b/Engine/Source/Editor/DataLayerEditor/Private/DataLayer/DataLayerMode.cpp
@@ -71,7 +71,7 @@ void FDataLayerMode::OnItemAdded(FSceneOutlinerTreeItemPtr Item)
 	{
 		if (!Item->Flags.bIsFilteredOut)
 		{
-			if (SelectedDataLayers.Contains(DataLayerItem->GetDataLayer()))
+			if (SelectedDataLayersSet.Contains(DataLayerItem->GetDataLayer()))
 			{
 				SceneOutliner->AddToSelection({Item});
 			}
@@ -610,9 +610,9 @@ TUniquePtr<ISceneOutlinerHierarchy> FDataLayerMode::CreateHierarchy()
 
 void FDataLayerMode::OnItemSelectionChanged(FSceneOutlinerTreeItemPtr TreeItem, ESelectInfo::Type SelectionType, const FSceneOutlinerItemSelection& Selection)
 {
-	SelectedDataLayers.Empty();
+	SelectedDataLayersSet.Empty();
 	SelectedDataLayerActors.Empty();
-	Selection.ForEachItem<FDataLayerTreeItem>([this](const FDataLayerTreeItem& Item) { SelectedDataLayers.Add(Item.GetDataLayer()); });
+	Selection.ForEachItem<FDataLayerTreeItem>([this](const FDataLayerTreeItem& Item) { SelectedDataLayersSet.Add(Item.GetDataLayer()); });
 	Selection.ForEachItem<FDataLayerActorTreeItem>([this](const FDataLayerActorTreeItem& Item) { SelectedDataLayerActors.Add(FSelectedDataLayerActor(Item.GetDataLayer(), Item.GetActor())); });
 }
 
diff --git a/Engine/Source/Editor/DataLayerEditor/Private/DataLayer/DataLayerMode.h b/Engine/Source/Editor/DataLayerEditor/Private/DataLayer/DataLayerMode.h
--- a/Engine/Source/Editor/DataLayerEditor/Private/DataLayer/DataLayerMode.h
+++ b/Engine/Source/Editor/DataLayerEditor/Private/DataLayer/DataLayerMode.h
@@ -80,7 +80,7 @@ private:
 	/** If this mode was created to display a specific world, don't allow it to be reassigned */
 	const TWeakObjectPtr<UWorld> SpecifiedWorldToDisplay;
 
-	TSet<TWeakObjectPtr<const UDataLayer>> SelectedDataLayers;
+	TSet<TWeakObjectPtr<const UDataLayer>> SelectedDataLayersSet;
 	typedef TPair<TWeakObjectPtr<const UDataLayer>, TWeakObjectPtr<const AActor>> FSelectedDataLayerActor;
 	TSet<FSelectedDataLayerActor> SelectedDataLayerActors;
 };
