DataLayers: Exposed most of DataLayerEditorSubsystem functionality to Blueprint and cleaned/removed redundant functions from API.

diff --git a/Engine/Source/Editor/DataLayerEditor/Private/DataLayer/DataLayerEditorSubsystem.cpp b/Engine/Source/Editor/DataLayerEditor/Private/DataLayer/DataLayerEditorSubsystem.cpp
--- a/Engine/Source/Editor/DataLayerEditor/Private/DataLayer/DataLayerEditorSubsystem.cpp
+++ b/Engine/Source/Editor/DataLayerEditor/Private/DataLayer/DataLayerEditorSubsystem.cpp
@@ -255,7 +255,7 @@ UWorld* UDataLayerEditorSubsystem::GetWorld() const
 	return GWorld;
 }
 
-bool UDataLayerEditorSubsystem::AddActorToDataLayer(AActor* Actor, const UDataLayer* DataLayer)
+bool UDataLayerEditorSubsystem::AddActorToDataLayer(AActor* Actor, UDataLayer* DataLayer)
 {
 	TArray<AActor*> Actors;
 	Actors.Add(Actor);
@@ -263,7 +263,7 @@ bool UDataLayerEditorSubsystem::AddActorToDataLayer(AActor* Actor, const UDataLa
 	return AddActorsToDataLayers(Actors, {DataLayer});
 }
 
-bool UDataLayerEditorSubsystem::AddActorToDataLayers(AActor* Actor, const TArray<const UDataLayer*>& DataLayers)
+bool UDataLayerEditorSubsystem::AddActorToDataLayers(AActor* Actor, const TArray<UDataLayer*>& DataLayers)
 {
 	TArray<AActor*> Actors;
 	Actors.Add(Actor);
@@ -271,17 +271,12 @@ bool UDataLayerEditorSubsystem::AddActorToDataLayers(AActor* Actor, const TArray
 	return AddActorsToDataLayers(Actors, DataLayers);
 }
 
-bool UDataLayerEditorSubsystem::AddActorsToDataLayer(const TArray<AActor*>& Actors, const UDataLayer* DataLayer)
+bool UDataLayerEditorSubsystem::AddActorsToDataLayer(const TArray<AActor*>& Actors, UDataLayer* DataLayer)
 {
 	return AddActorsToDataLayers(Actors, {DataLayer});
 }
 
-bool UDataLayerEditorSubsystem::AddActorsToDataLayer(const TArray<TWeakObjectPtr<AActor>>& Actors, const UDataLayer* DataLayer)
-{
-	return AddActorsToDataLayers(Actors, {DataLayer});
-}
-
-bool UDataLayerEditorSubsystem::AddActorsToDataLayers(const TArray<AActor*>& Actors, const TArray<const UDataLayer*>& DataLayers)
+bool UDataLayerEditorSubsystem::AddActorsToDataLayers(const TArray<AActor*>& Actors, const TArray<UDataLayer*>& DataLayers)
 {
 	bool bChangesOccurred = false;
 
@@ -326,15 +321,7 @@ bool UDataLayerEditorSubsystem::AddActorsToDataLayers(const TArray<AActor*>& Act
 	return bChangesOccurred;
 }
 
-bool UDataLayerEditorSubsystem::AddActorsToDataLayers(const TArray<TWeakObjectPtr<AActor>>& Actors, const TArray<const UDataLayer*>& DataLayers)
-{
-	TArray<AActor*> ActorsRawPtr;
-	Algo::TransformIf(Actors, ActorsRawPtr, [](const TWeakObjectPtr<AActor>& Actor) { return Actor.IsValid(); }, [](const TWeakObjectPtr<AActor>& Actor) { return Actor.Get(); });
-
-	return AddActorsToDataLayers(ActorsRawPtr, DataLayers);
-}
-
-bool UDataLayerEditorSubsystem::RemoveActorFromDataLayer(AActor* Actor, const UDataLayer* DataLayer)
+bool UDataLayerEditorSubsystem::RemoveActorFromDataLayer(AActor* Actor, UDataLayer* DataLayer)
 {
 	TArray<AActor*> Actors;
 	Actors.Add(Actor);
@@ -342,7 +329,7 @@ bool UDataLayerEditorSubsystem::RemoveActorFromDataLayer(AActor* Actor, const UD
 	return RemoveActorsFromDataLayers(Actors, {DataLayer});
 }
 
-bool UDataLayerEditorSubsystem::RemoveActorFromDataLayers(AActor* Actor, const TArray<const UDataLayer*>& DataLayers)
+bool UDataLayerEditorSubsystem::RemoveActorFromDataLayers(AActor* Actor, const TArray<UDataLayer*>& DataLayers)
 {
 	TArray<AActor*> Actors;
 	Actors.Add(Actor);
@@ -350,17 +337,12 @@ bool UDataLayerEditorSubsystem::RemoveActorFromDataLayers(AActor* Actor, const T
 	return RemoveActorsFromDataLayers(Actors, DataLayers);
 }
 
-bool UDataLayerEditorSubsystem::RemoveActorsFromDataLayer(const TArray<AActor*>& Actors, const UDataLayer* DataLayer)
-{
-	return RemoveActorsFromDataLayers(Actors, {DataLayer});
-}
-
-bool UDataLayerEditorSubsystem::RemoveActorsFromDataLayer(const TArray<TWeakObjectPtr<AActor>>& Actors, const UDataLayer* DataLayer)
+bool UDataLayerEditorSubsystem::RemoveActorsFromDataLayer(const TArray<AActor*>& Actors, UDataLayer* DataLayer)
 {
 	return RemoveActorsFromDataLayers(Actors, {DataLayer});
 }
 
-bool UDataLayerEditorSubsystem::RemoveActorsFromDataLayers(const TArray<AActor*>& Actors, const TArray<const UDataLayer*>& DataLayers)
+bool UDataLayerEditorSubsystem::RemoveActorsFromDataLayers(const TArray<AActor*>& Actors, const TArray<UDataLayer*>& DataLayers)
 {
 	GEditor->GetSelectedActors()->BeginBatchSelectOperation();
 
@@ -402,14 +384,6 @@ bool UDataLayerEditorSubsystem::RemoveActorsFromDataLayers(const TArray<AActor*>
 	return bChangesOccurred;
 }
 
-bool UDataLayerEditorSubsystem::RemoveActorsFromDataLayers(const TArray<TWeakObjectPtr<AActor>>& Actors, const TArray<const UDataLayer*>& DataLayers)
-{
-	TArray<AActor*> ActorsRawPtr;
-	Algo::TransformIf(Actors, ActorsRawPtr, [](const TWeakObjectPtr<AActor>& Actor) { return Actor.IsValid(); }, [](const TWeakObjectPtr<AActor>& Actor) { return Actor.Get(); });
-
-	return RemoveActorsFromDataLayers(ActorsRawPtr, DataLayers);
-}
-
 
 ///////////////////////////////////////////////////////////////////////////////////////////////////////////////////////
 //
@@ -423,22 +397,22 @@ TArray<AActor*> UDataLayerEditorSubsystem::GetSelectedActors() const
 	return CurrentlySelectedActors;
 }
 
-bool UDataLayerEditorSubsystem::AddSelectedActorsToDataLayer(const UDataLayer* DataLayer)
+bool UDataLayerEditorSubsystem::AddSelectedActorsToDataLayer(UDataLayer* DataLayer)
 {
 	return AddActorsToDataLayer(GetSelectedActors(), DataLayer);
 }
 
-bool UDataLayerEditorSubsystem::RemoveSelectedActorsFromDataLayer(const UDataLayer* DataLayer)
+bool UDataLayerEditorSubsystem::RemoveSelectedActorsFromDataLayer(UDataLayer* DataLayer)
 {
 	return RemoveActorsFromDataLayer(GetSelectedActors(), DataLayer);
 }
 
-bool UDataLayerEditorSubsystem::AddSelectedActorsToDataLayers(const TArray<const UDataLayer*>& DataLayers)
+bool UDataLayerEditorSubsystem::AddSelectedActorsToDataLayers(const TArray<UDataLayer*>& DataLayers)
 {
 	return AddActorsToDataLayers(GetSelectedActors(), DataLayers);
 }
 
-bool UDataLayerEditorSubsystem::RemoveSelectedActorsFromDataLayers(const TArray<const UDataLayer*>& DataLayers)
+bool UDataLayerEditorSubsystem::RemoveSelectedActorsFromDataLayers(const TArray<UDataLayer*>& DataLayers)
 {
 	return RemoveActorsFromDataLayers(GetSelectedActors(), DataLayers);
 }
@@ -449,12 +423,12 @@ bool UDataLayerEditorSubsystem::RemoveSelectedActorsFromDataLayers(const TArray<
 // Operations on actors in DataLayers
 //
 ///////////////////////////////////////////////////////////////////////////////////////////////////////////////////////
-bool UDataLayerEditorSubsystem::SelectActorsInDataLayer(const UDataLayer* DataLayer, const bool bSelect, const bool bNotify, const bool bSelectEvenIfHidden)
+bool UDataLayerEditorSubsystem::SelectActorsInDataLayer(UDataLayer* DataLayer, const bool bSelect, const bool bNotify, const bool bSelectEvenIfHidden)
 {
 	return SelectActorsInDataLayer(DataLayer, bSelect, bNotify, bSelectEvenIfHidden, nullptr);
 }
 
-bool UDataLayerEditorSubsystem::SelectActorsInDataLayer(const UDataLayer* DataLayer, const bool bSelect, const bool bNotify, const bool bSelectEvenIfHidden, const TSharedPtr<FActorFilter>& Filter)
+bool UDataLayerEditorSubsystem::SelectActorsInDataLayer(UDataLayer* DataLayer, const bool bSelect, const bool bNotify, const bool bSelectEvenIfHidden, const TSharedPtr<FActorFilter>& Filter)
 {
 	bool bChangesOccurred = false;
 
@@ -493,12 +467,12 @@ bool UDataLayerEditorSubsystem::SelectActorsInDataLayer(const UDataLayer* DataLa
 	return bChangesOccurred;
 }
 
-bool UDataLayerEditorSubsystem::SelectActorsInDataLayers(const TArray<const UDataLayer*>& DataLayers, const bool bSelect, const bool bNotify, const bool bSelectEvenIfHidden)
+bool UDataLayerEditorSubsystem::SelectActorsInDataLayers(const TArray<UDataLayer*>& DataLayers, const bool bSelect, const bool bNotify, const bool bSelectEvenIfHidden)
 {
 	return SelectActorsInDataLayers(DataLayers, bSelect, bNotify, bSelectEvenIfHidden, nullptr);
 }
 
-bool UDataLayerEditorSubsystem::SelectActorsInDataLayers(const TArray<const UDataLayer*>& DataLayers, const bool bSelect, const bool bNotify, const bool bSelectEvenIfHidden, const TSharedPtr<FActorFilter>& Filter)
+bool UDataLayerEditorSubsystem::SelectActorsInDataLayers(const TArray<UDataLayer*>& DataLayers, const bool bSelect, const bool bNotify, const bool bSelectEvenIfHidden, const TSharedPtr<FActorFilter>& Filter)
 {
 	if (DataLayers.Num() == 0)
 	{
@@ -550,7 +524,7 @@ bool UDataLayerEditorSubsystem::SelectActorsInDataLayers(const TArray<const UDat
 // Operations on actor viewport visibility regarding DataLayers
 //
 ///////////////////////////////////////////////////////////////////////////////////////////////////////////////////////
-void UDataLayerEditorSubsystem::UpdatePerViewVisibility(FLevelEditorViewportClient* ViewportClient, const UDataLayer* DataLayerThatChanged)
+void UDataLayerEditorSubsystem::UpdatePerViewVisibility(FLevelEditorViewportClient* ViewportClient, UDataLayer* DataLayerThatChanged)
 {
 	if (!ViewportClient->GetWorld())
 	{
@@ -588,7 +562,7 @@ void UDataLayerEditorSubsystem::UpdatePerViewVisibility(FLevelEditorViewportClie
 	ViewportClient->Invalidate();
 }
 
-void UDataLayerEditorSubsystem::UpdateAllViewVisibility(const UDataLayer* DataLayerThatChanged)
+void UDataLayerEditorSubsystem::UpdateAllViewVisibility(UDataLayer* DataLayerThatChanged)
 {
 	// Update all view's hidden DataLayers if they had this one
 	for (FLevelEditorViewportClient* ViewportClient : GEditor->GetLevelViewportClients())
@@ -861,12 +835,12 @@ bool UDataLayerEditorSubsystem::UpdateAllActorsVisibility(const bool bNotifySele
 //
 ///////////////////////////////////////////////////////////////////////////////////////////////////////////////////////
 
-void UDataLayerEditorSubsystem::AppendActorsFromDataLayer(const UDataLayer* DataLayer, TArray<AActor*>& InOutActors) const
+void UDataLayerEditorSubsystem::AppendActorsFromDataLayer(UDataLayer* DataLayer, TArray<AActor*>& InOutActors) const
 {
 	AppendActorsFromDataLayer(DataLayer, InOutActors, nullptr);
 }
 
-void UDataLayerEditorSubsystem::AppendActorsFromDataLayer(const UDataLayer* DataLayer, TArray<AActor*>& InOutActors, const TSharedPtr<FActorFilter>& Filter) const
+void UDataLayerEditorSubsystem::AppendActorsFromDataLayer(UDataLayer* DataLayer, TArray<AActor*>& InOutActors, const TSharedPtr<FActorFilter>& Filter) const
 {
 	for (AActor* Actor : FActorRange(GetWorld()))
 	{
@@ -881,27 +855,12 @@ void UDataLayerEditorSubsystem::AppendActorsFromDataLayer(const UDataLayer* Data
 	}
 }
 
-void UDataLayerEditorSubsystem::AppendActorsFromDataLayer(const UDataLayer* DataLayer, TArray<TWeakObjectPtr<AActor>>& InOutActors, const TSharedPtr<FActorFilter>& Filter) const
-{
-	for (AActor* Actor : FActorRange(GetWorld()))
-	{
-		if (Filter.IsValid() && !Filter->PassesFilter(Actor))
-		{
-			continue;
-		}
-		if (Actor->ContainsDataLayer(DataLayer))
-		{
-			InOutActors.Add(Actor);
-		}
-	}
-}
-
-void UDataLayerEditorSubsystem::AppendActorsFromDataLayers(const TArray<const UDataLayer*>& DataLayers, TArray<AActor*>& InOutActors) const
+void UDataLayerEditorSubsystem::AppendActorsFromDataLayers(const TArray<UDataLayer*>& DataLayers, TArray<AActor*>& InOutActors) const
 {
 	AppendActorsFromDataLayers(DataLayers, InOutActors, nullptr);
 }
 
-void UDataLayerEditorSubsystem::AppendActorsFromDataLayers(const TArray<const UDataLayer*>& DataLayers, TArray<AActor*>& InOutActors, const TSharedPtr<FActorFilter>& Filter) const
+void UDataLayerEditorSubsystem::AppendActorsFromDataLayers(const TArray<UDataLayer*>& DataLayers, TArray<AActor*>& InOutActors, const TSharedPtr<FActorFilter>& Filter) const
 {
 	for (AActor* Actor : FActorRange(GetWorld()))
 	{
@@ -920,47 +879,28 @@ void UDataLayerEditorSubsystem::AppendActorsFromDataLayers(const TArray<const UD
 	}
 }
 
-void UDataLayerEditorSubsystem::AppendActorsFromDataLayers(const TArray<const UDataLayer*>& DataLayers, TArray<TWeakObjectPtr<AActor>>& InOutActors, const TSharedPtr<FActorFilter>& Filter) const
-{
-	for (AActor* Actor : FActorRange(GetWorld()))
-	{
-		if (Filter.IsValid() && !Filter->PassesFilter(Actor))
-		{
-			continue;
-		}
-		for (const UDataLayer* DataLayer : DataLayers)
-		{
-			if (Actor->ContainsDataLayer(DataLayer))
-			{
-				InOutActors.Add(Actor);
-				break;
-			}
-		}
-	}
-}
-
-TArray<AActor*> UDataLayerEditorSubsystem::GetActorsFromDataLayer(const UDataLayer* DataLayer) const
+TArray<AActor*> UDataLayerEditorSubsystem::GetActorsFromDataLayer(UDataLayer* DataLayer) const
 {
 	TArray<AActor*> OutActors;
 	AppendActorsFromDataLayer(DataLayer, OutActors);
 	return OutActors;
 }
 
-TArray<AActor*> UDataLayerEditorSubsystem::GetActorsFromDataLayer(const UDataLayer* DataLayer, const TSharedPtr<FActorFilter>& Filter) const
+TArray<AActor*> UDataLayerEditorSubsystem::GetActorsFromDataLayer(UDataLayer* DataLayer, const TSharedPtr<FActorFilter>& Filter) const
 {
 	TArray<AActor*> OutActors;
 	AppendActorsFromDataLayer(DataLayer, OutActors, Filter);
 	return OutActors;
 }
 
-TArray<AActor*> UDataLayerEditorSubsystem::GetActorsFromDataLayers(const TArray<const UDataLayer*>& DataLayers) const
+TArray<AActor*> UDataLayerEditorSubsystem::GetActorsFromDataLayers(const TArray<UDataLayer*>& DataLayers) const
 {
 	TArray<AActor*> OutActors;
 	AppendActorsFromDataLayers(DataLayers, OutActors);
 	return OutActors;
 }
 
-TArray<AActor*> UDataLayerEditorSubsystem::GetActorsFromDataLayers(const TArray<const UDataLayer*>& DataLayers, const TSharedPtr<FActorFilter>& Filter) const
+TArray<AActor*> UDataLayerEditorSubsystem::GetActorsFromDataLayers(const TArray<UDataLayer*>& DataLayers, const TSharedPtr<FActorFilter>& Filter) const
 {
 	TArray<AActor*> OutActors;
 	AppendActorsFromDataLayers(DataLayers, OutActors, Filter);
@@ -1126,19 +1066,24 @@ bool UDataLayerEditorSubsystem::ToggleDataLayersIsDynamicallyLoadedInEditor(cons
 	return bRefreshNeeded ? RefreshWorldPartitionEditorCells() : true;
 }
 
-const UDataLayer* UDataLayerEditorSubsystem::GetDataLayerFromName(const FName& DataLayerName) const
+UDataLayer* UDataLayerEditorSubsystem::GetDataLayer(const FActorDataLayer& ActorDataLayer) const
+{
+	return GetDataLayerFromName(ActorDataLayer.Name);
+}
+
+UDataLayer* UDataLayerEditorSubsystem::GetDataLayerFromName(const FName& DataLayerName) const
 {
 	const AWorldDataLayers* WorldDataLayers = GetWorldDataLayers();
-	return WorldDataLayers ? WorldDataLayers->GetDataLayerFromName(DataLayerName) : nullptr;
+	return WorldDataLayers ? const_cast<UDataLayer*>(WorldDataLayers->GetDataLayerFromName(DataLayerName)) : nullptr;
 }
 
-const UDataLayer* UDataLayerEditorSubsystem::GetDataLayerFromLabel(const FName& DataLayerLabel) const
+UDataLayer* UDataLayerEditorSubsystem::GetDataLayerFromLabel(const FName& DataLayerLabel) const
 {
 	const AWorldDataLayers* WorldDataLayers = GetWorldDataLayers();
-	return WorldDataLayers ? WorldDataLayers->GetDataLayerFromLabel(DataLayerLabel) : nullptr;
+	return WorldDataLayers ? const_cast<UDataLayer*>(WorldDataLayers->GetDataLayerFromLabel(DataLayerLabel)) : nullptr;
 }
 
-bool UDataLayerEditorSubsystem::TryGetDataLayerFromLabel(const FName& DataLayerLabel, const UDataLayer*& OutDataLayer)
+bool UDataLayerEditorSubsystem::TryGetDataLayerFromLabel(const FName& DataLayerLabel, UDataLayer*& OutDataLayer)
 {
 	OutDataLayer = GetDataLayerFromLabel(DataLayerLabel);
 	return (OutDataLayer != nullptr);
@@ -1174,7 +1119,7 @@ UDataLayer* UDataLayerEditorSubsystem::CreateDataLayer()
 	return NewDataLayer;
 }
 
-void UDataLayerEditorSubsystem::DeleteDataLayers(const TArray<const UDataLayer*>& DataLayersToDelete)
+void UDataLayerEditorSubsystem::DeleteDataLayers(const TArray<UDataLayer*>& DataLayersToDelete)
 {
 	if (AWorldDataLayers* WorldDataLayers = GetWorldDataLayers())
 	{
@@ -1185,7 +1130,7 @@ void UDataLayerEditorSubsystem::DeleteDataLayers(const TArray<const UDataLayer*>
 	}
 }
 
-void UDataLayerEditorSubsystem::DeleteDataLayer(const UDataLayer* DataLayerToDelete)
+void UDataLayerEditorSubsystem::DeleteDataLayer(UDataLayer* DataLayerToDelete)
 {
 	if (AWorldDataLayers* WorldDataLayers = GetWorldDataLayers())
 	{
diff --git a/Engine/Source/Editor/DataLayerEditor/Private/DataLayer/DataLayerMode.cpp b/Engine/Source/Editor/DataLayerEditor/Private/DataLayer/DataLayerMode.cpp
--- a/Engine/Source/Editor/DataLayerEditor/Private/DataLayer/DataLayerMode.cpp
+++ b/Engine/Source/Editor/DataLayerEditor/Private/DataLayer/DataLayerMode.cpp
@@ -90,7 +90,7 @@ void FDataLayerMode::OnItemDoubleClick(FSceneOutlinerTreeItemPtr Item)
 {
 	if (const FDataLayerTreeItem* DataLayerItem = Item->CastTo<FDataLayerTreeItem>())
 	{
-		if (const UDataLayer* DataLayer = DataLayerItem->GetDataLayer())
+		if (UDataLayer* DataLayer = DataLayerItem->GetDataLayer())
 		{
 			const FScopedTransaction Transaction(LOCTEXT("SelectActorsInDataLayer", "Select Actors in DataLayer"));
 			GEditor->SelectNone(/*bNoteSelectionChange*/false, true);
@@ -113,14 +113,14 @@ void FDataLayerMode::OnItemDoubleClick(FSceneOutlinerTreeItemPtr Item)
 
 void FDataLayerMode::DeleteItems(const TArray<TWeakPtr<ISceneOutlinerTreeItem>>& Items)
 {
-	TArray<const UDataLayer*> DataLayersToDelete;
-	TMap<const UDataLayer*, TArray<AActor*>> ActorsToRemoveFromDataLayer;
+	TArray<UDataLayer*> DataLayersToDelete;
+	TMap<UDataLayer*, TArray<AActor*>> ActorsToRemoveFromDataLayer;
 
 	for (const TWeakPtr<ISceneOutlinerTreeItem>& Item : Items)
 	{
 		if (FDataLayerActorTreeItem* DataLayerActorItem = Item.Pin()->CastTo<FDataLayerActorTreeItem>())
 		{
-			const UDataLayer* DataLayer = DataLayerActorItem->GetDataLayer();
+			UDataLayer* DataLayer = DataLayerActorItem->GetDataLayer();
 			AActor* Actor = DataLayerActorItem->GetActor();
 			if (DataLayer && Actor)
 			{
@@ -293,11 +293,11 @@ void FDataLayerMode::OnDrop(ISceneOutlinerTreeItem& DropTarget, const FSceneOutl
 
 	if (const FDataLayerTreeItem* DataLayerItem = DropTarget.CastTo<FDataLayerTreeItem>())
 	{
-		if (const UDataLayer* DataLayer = DataLayerItem->GetDataLayer())
+		if (UDataLayer* DataLayer = DataLayerItem->GetDataLayer())
 		{
 			if (SceneOutliner->GetTree().IsItemSelected(const_cast<ISceneOutlinerTreeItem&>(DropTarget).AsShared()))
 			{
-				TArray<const UDataLayer*> AllSelectedDataLayers = GetSelectedDataLayers(SceneOutliner);
+				TArray<UDataLayer*> AllSelectedDataLayers = GetSelectedDataLayers(SceneOutliner);
 				if (AllSelectedDataLayers.Num() > 1)
 				{
 					const FScopedTransaction Transaction(LOCTEXT("DataLayerOutlinerAddActorsToDataLayers", "Add Actors to DataLayers"));
@@ -316,7 +316,7 @@ void FDataLayerMode::OnDrop(ISceneOutlinerTreeItem& DropTarget, const FSceneOutl
 		{
 			// Only allow actors not coming from the DataLayerBrowser
 			const FScopedTransaction Transaction(LOCTEXT("AddSelectedActorsToNewDataLayer", "Add Actors to New DataLayer"));
-			if (const UDataLayer* NewDataLayer = DataLayerEditorSubsystem->CreateDataLayer())
+			if (UDataLayer* NewDataLayer = DataLayerEditorSubsystem->CreateDataLayer())
 			{
 				DataLayerEditorSubsystem->AddActorsToDataLayer(ActorsToAdd, NewDataLayer);
 			}
@@ -397,12 +397,12 @@ TSharedPtr<FDragDropOperation> FDataLayerMode::CreateDragDropOperation(const TAr
 static const FName DefaultContextBaseMenuName("DataLayerOutliner.DefaultContextMenuBase");
 static const FName DefaultContextMenuName("DataLayerOutliner.DefaultContextMenu");
 
-TArray<const UDataLayer*> FDataLayerMode::GetSelectedDataLayers(SSceneOutliner* InSceneOutliner) const
+TArray<UDataLayer*> FDataLayerMode::GetSelectedDataLayers(SSceneOutliner* InSceneOutliner) const
 {
 	FSceneOutlinerItemSelection ItemSelection(InSceneOutliner->GetSelection());
 	TArray<FDataLayerTreeItem*> SelectedDataLayerItems;
 	ItemSelection.Get<FDataLayerTreeItem>(SelectedDataLayerItems);
-	TArray<const UDataLayer*> ValidSelectedDataLayers;
+	TArray<UDataLayer*> ValidSelectedDataLayers;
 	Algo::TransformIf(SelectedDataLayerItems, ValidSelectedDataLayers, [](const auto Item) { return Item && Item->GetDataLayer(); }, [](const auto Item) { return Item->GetDataLayer(); });
 	return MoveTemp(ValidSelectedDataLayers);
 }
@@ -424,7 +424,7 @@ void FDataLayerMode::RegisterContextMenu()
 			}
 
 			SSceneOutliner* SceneOutliner = Context->SceneOutliner.Pin().Get();
-			TArray<const UDataLayer*> SelectedDataLayers = GetSelectedDataLayers(SceneOutliner);
+			TArray<UDataLayer*> SelectedDataLayers = GetSelectedDataLayers(SceneOutliner);
 
 			TArray<const UDataLayer*> AllDataLayers;
 			if (const AWorldDataLayers* WorldDataLayers = AWorldDataLayers::Get(RepresentingWorld.Get()))
@@ -455,7 +455,7 @@ void FDataLayerMode::RegisterContextMenu()
 								DataLayerEditorSubsystem->AddSelectedActorsToDataLayer(NewDataLayer);
 							}}),
 						FCanExecuteAction::CreateLambda([=] { return GEditor->GetSelectedActorCount() > 0; })
-								));
+					));
 
 				Section.AddMenuEntry("AddSelectedActorsToSelectedDataLayers", LOCTEXT("AddSelectedActorsToSelectedDataLayers", "Add Selected Actors to Selected DataLayers"), FText(), FSlateIcon(),
 					FUIAction(
@@ -466,7 +466,7 @@ void FDataLayerMode::RegisterContextMenu()
 								DataLayerEditorSubsystem->AddSelectedActorsToDataLayers(SelectedDataLayers);
 							}}),
 						FCanExecuteAction::CreateLambda([=] { return !SelectedDataLayers.IsEmpty() && GEditor->GetSelectedActorCount() > 0; })
-								));
+					));
 
 				Section.AddSeparator("SectionsSeparator");
 
@@ -479,7 +479,7 @@ void FDataLayerMode::RegisterContextMenu()
 								DataLayerEditorSubsystem->RemoveSelectedActorsFromDataLayers(SelectedDataLayers);
 							}}),
 						FCanExecuteAction::CreateLambda([=] { return !SelectedDataLayers.IsEmpty() && GEditor->GetSelectedActorCount() > 0; })
-								));
+					));
 
 				Section.AddMenuEntry("DeleteSelectedDataLayers", LOCTEXT("DeleteSelectedDataLayers", "Delete Selected DataLayers"), FText(), FSlateIcon(),
 					FUIAction(
@@ -490,7 +490,7 @@ void FDataLayerMode::RegisterContextMenu()
 								DataLayerEditorSubsystem->DeleteDataLayers(SelectedDataLayers);
 							}}),
 						FCanExecuteAction::CreateLambda([=] { return !SelectedDataLayers.IsEmpty(); })
-								));
+					));
 
 				Section.AddMenuEntry("RenameSelectedDataLayer", LOCTEXT("RenameSelectedDataLayer", "Rename Selected DataLayer"), FText(), FSlateIcon(),
 					FUIAction(
@@ -505,7 +505,7 @@ void FDataLayerMode::RegisterContextMenu()
 								}
 							}}),
 						FCanExecuteAction::CreateLambda([=] { return SelectedDataLayers.Num() == 1; })
-								));
+					));
 
 				Section.AddSeparator("SectionsSeparator");
 			}
@@ -523,7 +523,7 @@ void FDataLayerMode::RegisterContextMenu()
 								DataLayerEditorSubsystem->SelectActorsInDataLayers(SelectedDataLayers, /*bSelect*/true, /*bNotify*/true, /*bSelectEvenIfHidden*/true);
 							}}),
 						FCanExecuteAction::CreateLambda([=] { return !SelectedDataLayers.IsEmpty(); })
-								));
+					));
 
 				Section.AddMenuEntry("AppendActorsToSelection", LOCTEXT("AppendActorsToSelection", "Append Actors in DataLayer to Selection"), FText(), FSlateIcon(),
 					FUIAction(
@@ -534,7 +534,7 @@ void FDataLayerMode::RegisterContextMenu()
 								DataLayerEditorSubsystem->SelectActorsInDataLayers(SelectedDataLayers, /*bSelect*/true, /*bNotify*/true, /*bSelectEvenIfHidden*/true);
 							}}),
 						FCanExecuteAction::CreateLambda([=] { return !SelectedDataLayers.IsEmpty(); })
-								));
+					));
 
 				Section.AddMenuEntry("DeselectActors", LOCTEXT("DeselectActors", "Deselect Actors in DataLayer"), FText(), FSlateIcon(),
 					FUIAction(
@@ -545,7 +545,7 @@ void FDataLayerMode::RegisterContextMenu()
 								DataLayerEditorSubsystem->SelectActorsInDataLayers(SelectedDataLayers, /*bSelect*/false, /*bNotifySelectActors*/true);
 							}}),
 						FCanExecuteAction::CreateLambda([=] { return !SelectedDataLayers.IsEmpty(); })
-								));
+					));
 			}
 
 			{
@@ -560,9 +560,7 @@ void FDataLayerMode::RegisterContextMenu()
 								DataLayerEditorSubsystem->MakeAllDataLayersVisible();
 							}}),
 						FCanExecuteAction::CreateLambda([=] { return !AllDataLayers.IsEmpty(); })
-								));
-
-
+					));
 			}
 		}));
 	}
diff --git a/Engine/Source/Editor/DataLayerEditor/Private/DataLayer/DataLayerMode.h b/Engine/Source/Editor/DataLayerEditor/Private/DataLayer/DataLayerMode.h
--- a/Engine/Source/Editor/DataLayerEditor/Private/DataLayer/DataLayerMode.h
+++ b/Engine/Source/Editor/DataLayerEditor/Private/DataLayer/DataLayerMode.h
@@ -67,7 +67,7 @@ private:
 	void RegisterContextMenu();
 	void ChooseRepresentingWorld();
 	TArray<AActor*> GetActorsFromOperation(const FDragDropOperation& Operation, bool bOnlyFindFirst = false) const;
-	TArray<const UDataLayer*> GetSelectedDataLayers(SSceneOutliner* InSceneOutliner) const;
+	TArray<UDataLayer*> GetSelectedDataLayers(SSceneOutliner* InSceneOutliner) const;
 
 	/** The world which we are currently representing */
 	TWeakObjectPtr<UWorld> RepresentingWorld;
diff --git a/Engine/Source/Editor/DataLayerEditor/Private/DataLayer/DataLayerPropertyTypeCustomization.cpp b/Engine/Source/Editor/DataLayerEditor/Private/DataLayer/DataLayerPropertyTypeCustomization.cpp
--- a/Engine/Source/Editor/DataLayerEditor/Private/DataLayer/DataLayerPropertyTypeCustomization.cpp
+++ b/Engine/Source/Editor/DataLayerEditor/Private/DataLayer/DataLayerPropertyTypeCustomization.cpp
@@ -96,13 +96,12 @@ FText FDataLayerPropertyTypeCustomization::GetDataLayerDescription(const UDataLa
 	return FText::FromName(InDataLayer->GetDataLayerLabel());
 }
 
-const UDataLayer* FDataLayerPropertyTypeCustomization::GetDataLayerFromPropertyHandle() const
+UDataLayer* FDataLayerPropertyTypeCustomization::GetDataLayerFromPropertyHandle() const
 {
 	FName DataLayerName;
 	if (PropertyHandle->GetValue(DataLayerName) == FPropertyAccess::Success)
 	{
-		UDataLayerEditorSubsystem* DataLayerSubsystem = UDataLayerEditorSubsystem::Get();
-		const UDataLayer* DataLayer = DataLayerSubsystem->GetDataLayerFromName(DataLayerName);
+		UDataLayer* DataLayer = UDataLayerEditorSubsystem::Get()->GetDataLayerFromName(DataLayerName);
 		return DataLayer;
 	}
 	return nullptr;
@@ -130,8 +129,7 @@ TSharedRef<SWidget> FDataLayerPropertyTypeCustomization::OnGetDataLayerMenu()
 	MenuBuilder.BeginSection(FName(), LOCTEXT("ExistingDataLayers", "Existing DataLayers"));
 	{
 		TArray<TWeakObjectPtr<UDataLayer>> AllDataLayers;
-		UDataLayerEditorSubsystem* DataLayerSubsystem = UDataLayerEditorSubsystem::Get();
-		DataLayerSubsystem->AddAllDataLayersTo(AllDataLayers);
+		UDataLayerEditorSubsystem::Get()->AddAllDataLayersTo(AllDataLayers);
 
 		for (const TWeakObjectPtr<UDataLayer>& WeakDataLayer : AllDataLayers)
 		{
@@ -161,7 +159,7 @@ EVisibility FDataLayerPropertyTypeCustomization::GetSelectDataLayerVisibility()
 
 FReply FDataLayerPropertyTypeCustomization::OnSelectDataLayer()
 {
-	if (const UDataLayer* DataLayer = GetDataLayerFromPropertyHandle())
+	if (UDataLayer* DataLayer = GetDataLayerFromPropertyHandle())
 	{
 		GEditor->SelectNone(true, true);
 		UDataLayerEditorSubsystem::Get()->SelectActorsInDataLayer(DataLayer, true, true);
diff --git a/Engine/Source/Editor/DataLayerEditor/Private/DataLayer/DataLayerPropertyTypeCustomization.h b/Engine/Source/Editor/DataLayerEditor/Private/DataLayer/DataLayerPropertyTypeCustomization.h
--- a/Engine/Source/Editor/DataLayerEditor/Private/DataLayer/DataLayerPropertyTypeCustomization.h
+++ b/Engine/Source/Editor/DataLayerEditor/Private/DataLayer/DataLayerPropertyTypeCustomization.h
@@ -21,7 +21,7 @@ private:
 	void AssignDataLayer(const UDataLayer* InDataLayer);
 	void OpenDataLayerBrowser();
 
-	const UDataLayer* GetDataLayerFromPropertyHandle() const;
+	UDataLayer* GetDataLayerFromPropertyHandle() const;
 	FText GetDataLayerDescription(const UDataLayer* InDataLayer) const;
 	FText GetDataLayerText() const;
 	EVisibility GetSelectDataLayerVisibility() const;
diff --git a/Engine/Source/Editor/DataLayerEditor/Private/DataLayer/DataLayerTreeItem.h b/Engine/Source/Editor/DataLayerEditor/Private/DataLayer/DataLayerTreeItem.h
--- a/Engine/Source/Editor/DataLayerEditor/Private/DataLayer/DataLayerTreeItem.h
+++ b/Engine/Source/Editor/DataLayerEditor/Private/DataLayer/DataLayerTreeItem.h
@@ -11,8 +11,7 @@
 struct FDataLayerTreeItem : ISceneOutlinerTreeItem
 {
 	FDataLayerTreeItem(UDataLayer* InDataLayer);
-	const UDataLayer* GetDataLayer() const { return DataLayer.Get(); }
-	UDataLayer* GetDataLayer() { return DataLayer.Get(); }
+	UDataLayer* GetDataLayer() const { return DataLayer.Get(); }
 
 	/* Begin ISceneOutlinerTreeItem Implementation */
 	virtual bool IsValid() const override { return DataLayer.IsValid(); }
diff --git a/Engine/Source/Editor/DataLayerEditor/Private/DataLayer/SDataLayerTreeLabel.cpp b/Engine/Source/Editor/DataLayerEditor/Private/DataLayer/SDataLayerTreeLabel.cpp
--- a/Engine/Source/Editor/DataLayerEditor/Private/DataLayer/SDataLayerTreeLabel.cpp
+++ b/Engine/Source/Editor/DataLayerEditor/Private/DataLayer/SDataLayerTreeLabel.cpp
@@ -155,7 +155,7 @@ bool SDataLayerTreeLabel::OnVerifyItemLabelChanged(const FText& InLabel, FText&
 		return false;
 	}
 
-	const UDataLayer* FoundDataLayer;
+	UDataLayer* FoundDataLayer;
 	if (UDataLayerEditorSubsystem::Get()->TryGetDataLayerFromLabel(*InLabel.ToString(), FoundDataLayer) && FoundDataLayer != DataLayerPtr.Get())
 	{
 		OutErrorMessage = LOCTEXT("RenameFailed_AlreadyExists", "This DataLayer already exists");
diff --git a/Engine/Source/Editor/DataLayerEditor/Public/DataLayer/DataLayerEditorSubsystem.h b/Engine/Source/Editor/DataLayerEditor/Public/DataLayer/DataLayerEditorSubsystem.h
--- a/Engine/Source/Editor/DataLayerEditor/Public/DataLayer/DataLayerEditorSubsystem.h
+++ b/Engine/Source/Editor/DataLayerEditor/Public/DataLayer/DataLayerEditorSubsystem.h
@@ -66,15 +66,9 @@ public:
 	 *
 	 *	@param	Actor	The actor to validate
 	 */
+	UFUNCTION(BlueprintCallable, Category = DataLayers)
 	virtual bool IsActorValidForDataLayer(AActor* Actor);
 
-	/**
-	 *	Synchronizes an newly created Actor's DataLayers with the DataLayer system
-	 *
-	 *	@param	Actor	The actor to initialize
-	 */
-	virtual void InitializeNewActorDataLayers(AActor* Actor);
-	
 	/**
 	 * Adds the actor to the DataLayer.
 	 *
@@ -82,7 +76,8 @@ public:
 	 * @param	DataLayer	The DataLayer to add the actor to
 	 * @return				true if the actor was added.  false is returned if the actor already belongs to the DataLayer.
 	 */
-	virtual bool AddActorToDataLayer(AActor* Actor, const UDataLayer* DataLayer);
+	UFUNCTION(BlueprintCallable, Category = DataLayers)
+	virtual bool AddActorToDataLayer(AActor* Actor, UDataLayer* DataLayer);
 
 	/**
 	 * Adds the provided actor to the DataLayers.
@@ -91,7 +86,8 @@ public:
 	 * @param	DataLayers	A valid list of DataLayers.
 	 * @return				true if the actor was added to at least one of the provided DataLayers.
 	 */
-	virtual bool AddActorToDataLayers(AActor* Actor, const TArray<const UDataLayer*>& DataLayers);
+	UFUNCTION(BlueprintCallable, Category = DataLayers)
+	virtual bool AddActorToDataLayers(AActor* Actor, const TArray<UDataLayer*>& DataLayers);
 
 	/**
 	 * Removes an actor from the specified DataLayer.
@@ -100,7 +96,8 @@ public:
 	 * @param	DataLayerToRemove	The DataLayer to remove the actor from
 	 * @return					true if the actor was removed from the DataLayer.  false is returned if the actor already belonged to the DataLayer.
 	 */
-	virtual bool RemoveActorFromDataLayer(AActor* Actor, const UDataLayer* DataLayerToRemove);
+	UFUNCTION(BlueprintCallable, Category = DataLayers)
+	virtual bool RemoveActorFromDataLayer(AActor* Actor, UDataLayer* DataLayerToRemove);
 
 	/**
 	 * Removes the provided actor from the DataLayers.
@@ -109,7 +106,8 @@ public:
 	 * @param	DataLayers	A valid list of DataLayers.
 	 * @return				true if the actor was removed from at least one of the provided DataLayers.
 	 */
-	virtual bool RemoveActorFromDataLayers(AActor* Actor, const TArray<const UDataLayer*>& DataLayers);
+	UFUNCTION(BlueprintCallable, Category = DataLayers)
+	virtual bool RemoveActorFromDataLayers(AActor* Actor, const TArray<UDataLayer*>& DataLayers);
 
 
 	/////////////////////////////////////////////////
@@ -122,16 +120,8 @@ public:
 	 * @param	DataLayer	The DataLayer to add to
 	 * @return				true if at least one actor was added to the DataLayer.  false is returned if all the actors already belonged to the DataLayer.
 	 */
-	virtual bool AddActorsToDataLayer(const TArray<AActor*>& Actors, const UDataLayer* DataLayer);
-
-	/**
-	 * Add the actors to the DataLayer
-	 *
-	 * @param	Actors		The actors to add to the DataLayer
-	 * @param	DataLayer	The DataLayer to add to
-	 * @return				true if at least one actor was added to the DataLayer.  false is returned if all the actors already belonged to the DataLayer.
-	 */
-	virtual bool AddActorsToDataLayer(const TArray<TWeakObjectPtr<AActor>>& Actors, const UDataLayer* DataLayer);
+	UFUNCTION(BlueprintCallable, Category = DataLayers)
+	virtual bool AddActorsToDataLayer(const TArray<AActor*>& Actors, UDataLayer* DataLayer);
 
 	/**
 	 * Add the actors to the DataLayers
@@ -140,16 +130,8 @@ public:
 	 * @param	DataLayers	A valid list of DataLayers.
 	 * @return				true if at least one actor was added to at least one DataLayer.  false is returned if all the actors already belonged to all specified DataLayers.
 	 */
-	bool AddActorsToDataLayers(const TArray<AActor*>& Actors, const TArray<const UDataLayer*>& DataLayers);
-
-	/**
-	 * Add the actors to the DataLayers
-	 *
-	 * @param	Actors		The actors to add to the DataLayers
-	 * @param	DataLayers	A valid list of DataLayers.
-	 * @return				true if at least one actor was added to at least one DataLayer.  false is returned if all the actors already belonged to all specified DataLayers.
-	 */
-	virtual bool AddActorsToDataLayers(const TArray<TWeakObjectPtr<AActor>>& Actors, const TArray<const UDataLayer*>& DataLayers);
+	UFUNCTION(BlueprintCallable, Category = DataLayers)
+	bool AddActorsToDataLayers(const TArray<AActor*>& Actors, const TArray<UDataLayer*>& DataLayers);
 
 	/**
 	 * Removes the actors from the specified DataLayer.
@@ -158,25 +140,8 @@ public:
 	 * @param	DataLayerToRemove	The DataLayer to remove the actors from
 	 * @return					true if at least one actor was removed from the DataLayer.  false is returned if all the actors already belonged to the DataLayer.
 	 */
-	virtual bool RemoveActorsFromDataLayer(const TArray<AActor*>& Actors, const UDataLayer* DataLayer);
-
-	/**
-	 * Removes the actors from the specified DataLayer.
-	 *
-	 * @param	Actors			The actors to remove from the provided DataLayer
-	 * @param	DataLayerToRemove	The DataLayer to remove the actors from
-	 * @return					true if at least one actor was removed from the DataLayer.  false is returned if all the actors already belonged to the DataLayer.
-	 */
-	virtual bool RemoveActorsFromDataLayer(const TArray<TWeakObjectPtr<AActor>>& Actors, const UDataLayer* DataLayer);
-
-	/**
-	 * Remove the actors to the DataLayers
-	 *
-	 * @param	Actors		The actors to remove to the DataLayers
-	 * @param	DataLayers	A valid list of DataLayers.
-	 * @return				true if at least one actor was removed from at least one DataLayer.  false is returned if none of the actors belonged to any of the specified DataLayers.
-	 */
-	virtual bool RemoveActorsFromDataLayers(const TArray<AActor*>& Actors, const TArray<const UDataLayer*>& DataLayers);
+	UFUNCTION(BlueprintCallable, Category = DataLayers)
+	virtual bool RemoveActorsFromDataLayer(const TArray<AActor*>& Actors, UDataLayer* DataLayer);
 
 	/**
 	 * Remove the actors to the DataLayers
@@ -185,26 +150,20 @@ public:
 	 * @param	DataLayers	A valid list of DataLayers.
 	 * @return				true if at least one actor was removed from at least one DataLayer.  false is returned if none of the actors belonged to any of the specified DataLayers.
 	 */
-	virtual bool RemoveActorsFromDataLayers(const TArray<TWeakObjectPtr<AActor>>& Actors, const TArray<const UDataLayer*>& DataLayers);
+	UFUNCTION(BlueprintCallable, Category = DataLayers)
+	virtual bool RemoveActorsFromDataLayers(const TArray<AActor*>& Actors, const TArray<UDataLayer*>& DataLayers);
 
-	
 	/////////////////////////////////////////////////
 	// Operations on selected actors.
 
-	/**
-	 * Find and return the selected actors.
-	 *
-	 * @return				The selected AActor's as a TArray.
-	 */
-	virtual TArray<AActor*> GetSelectedActors() const;
-
 	/**
 	 * Adds selected actors to the DataLayer.
 	 *
 	 * @param	DataLayer	A DataLayer.
 	 * @return				true if at least one actor was added.  false is returned if all selected actors already belong to the DataLayer.
 	 */
-	virtual bool AddSelectedActorsToDataLayer(const UDataLayer* DataLayer);
+	UFUNCTION(BlueprintCallable, Category = DataLayers)
+	virtual bool AddSelectedActorsToDataLayer(UDataLayer* DataLayer);
 
 	/**
 	 * Adds selected actors to the DataLayers.
@@ -212,7 +171,8 @@ public:
 	 * @param	DataLayers	A valid list of DataLayers.
 	 * @return				true if at least one actor was added.  false is returned if all selected actors already belong to the DataLayers.
 	 */
-	virtual bool AddSelectedActorsToDataLayers(const TArray<const UDataLayer*>& DataLayers);
+	UFUNCTION(BlueprintCallable, Category = DataLayers)
+	virtual bool AddSelectedActorsToDataLayers(const TArray<UDataLayer*>& DataLayers);
 
 	/**
 	 * Removes the selected actors from the DataLayer.
@@ -220,7 +180,8 @@ public:
 	 * @param	DataLayer	A DataLayer.
 	 * @return				true if at least one actor was added.  false is returned if all selected actors already belong to the DataLayer.
 	 */
-	virtual bool RemoveSelectedActorsFromDataLayer(const UDataLayer* DataLayer);
+	UFUNCTION(BlueprintCallable, Category = DataLayers)
+	virtual bool RemoveSelectedActorsFromDataLayer(UDataLayer* DataLayer);
 
 	/**
 	 * Removes selected actors from the DataLayers.
@@ -228,7 +189,8 @@ public:
 	 * @param	DataLayers	A valid list of DataLayers.
 	 * @return				true if at least one actor was removed.
 	 */
-	virtual bool RemoveSelectedActorsFromDataLayers(const TArray<const UDataLayer*>& DataLayers);
+	UFUNCTION(BlueprintCallable, Category = DataLayers)
+	virtual bool RemoveSelectedActorsFromDataLayers(const TArray<UDataLayer*>& DataLayers);
 
 
 	/////////////////////////////////////////////////
@@ -243,7 +205,8 @@ public:
 	 * @param	bSelectEvenIfHidden	[optional]	If true even hidden actors will be selected; if false, hidden actors won't be selected.
 	 * @return									true if at least one actor was selected/deselected.
 	 */
-	virtual bool SelectActorsInDataLayer(const UDataLayer* DataLayer, const bool bSelect, const bool bNotify, const bool bSelectEvenIfHidden = false);
+	UFUNCTION(BlueprintCallable, Category = DataLayers)
+	virtual bool SelectActorsInDataLayer(UDataLayer* DataLayer, const bool bSelect, const bool bNotify, const bool bSelectEvenIfHidden = false);
 
 	/**
 	 * Selects/de-selects actors belonging to the DataLayer.
@@ -255,7 +218,8 @@ public:
 	 * @param	Filter	[optional]				Actor that don't pass the specified filter restrictions won't be selected.
 	 * @return									true if at least one actor was selected/deselected.
 	 */
-	virtual bool SelectActorsInDataLayer(const UDataLayer* DataLayer, const bool bSelect, const bool bNotify, const bool bSelectEvenIfHidden, const TSharedPtr<FActorFilter>& Filter);
+	virtual bool SelectActorsInDataLayer(UDataLayer* DataLayer, const bool bSelect, const bool bNotify, const bool bSelectEvenIfHidden, const TSharedPtr<FActorFilter>& Filter);
+
 	/**
 	 * Selects/de-selects actors belonging to the DataLayers.
 	 *
@@ -265,7 +229,8 @@ public:
 	 * @param	bSelectEvenIfHidden	[optional]	If true even hidden actors will be selected; if false, hidden actors won't be selected.
 	 * @return									true if at least one actor was selected/deselected.
 	 */
-	virtual bool SelectActorsInDataLayers(const TArray<const UDataLayer*>& DataLayers, const bool bSelect, const bool bNotify, const bool bSelectEvenIfHidden = false);
+	UFUNCTION(BlueprintCallable, Category = DataLayers)
+	virtual bool SelectActorsInDataLayers(const TArray<UDataLayer*>& DataLayers, const bool bSelect, const bool bNotify, const bool bSelectEvenIfHidden = false);
 
 	/**
 	 * Selects/de-selects actors belonging to the DataLayers.
@@ -277,7 +242,7 @@ public:
 	 * @param	Filter	[optional]				Actor that don't pass the specified filter restrictions won't be selected.
 	 * @return									true if at least one actor was selected/deselected.
 	 */
-	virtual bool SelectActorsInDataLayers(const TArray<const UDataLayer*>& DataLayers, const bool bSelect, const bool bNotify, const bool bSelectEvenIfHidden, const TSharedPtr<FActorFilter>& Filter);
+	virtual bool SelectActorsInDataLayers(const TArray<UDataLayer*>& DataLayers, const bool bSelect, const bool bNotify, const bool bSelectEvenIfHidden, const TSharedPtr<FActorFilter>& Filter);
 
 
 	/////////////////////////////////////////////////
@@ -288,7 +253,8 @@ public:
 	 *
 	 * @param DataLayerThatChanged  If one DataLayer was changed (toggled in view pop-up, etc), then we only need to modify actors that use that DataLayer.
 	 */
-	virtual void UpdateAllViewVisibility(const UDataLayer* DataLayerThatChanged);
+	UFUNCTION(BlueprintCallable, Category = DataLayers)
+	virtual void UpdateAllViewVisibility(UDataLayer* DataLayerThatChanged);
 
 	/**
 	 * Updates the per-view visibility for all actors for the given view
@@ -296,7 +262,7 @@ public:
 	 * @param ViewportClient				The viewport client to update visibility on
 	 * @param DataLayerThatChanged [optional]	If one DataLayer was changed (toggled in view pop-up, etc), then we only need to modify actors that use that DataLayer
 	 */
-	virtual void UpdatePerViewVisibility(FLevelEditorViewportClient* ViewportClient, const UDataLayer* DataLayerThatChanged = nullptr);
+	virtual void UpdatePerViewVisibility(FLevelEditorViewportClient* ViewportClient, UDataLayer* DataLayerThatChanged = nullptr);
 
 	/**
 	 * Updates per-view visibility for the given actor in the given view
@@ -312,6 +278,7 @@ public:
 	 *
 	 * @param Actor		Actor to update
 	 */
+	UFUNCTION(BlueprintCallable, Category = DataLayers)
 	virtual void UpdateActorAllViewsVisibility(AActor* Actor);
 
 	/**
@@ -330,6 +297,7 @@ public:
 	 * @param	bNotifySelectionChange		If true the Editor is notified of the selection change; if false, the Editor will not be notified
 	 * @param	bRedrawViewports			If true the viewports will be redrawn; if false, they will not
 	 */
+	UFUNCTION(BlueprintCallable, Category = DataLayers)
 	virtual bool UpdateActorVisibility(AActor* Actor, bool& bOutSelectionChanged, bool& bOutActorModified, const bool bNotifySelectionChange, const bool bRedrawViewports);
 
 	/**
@@ -338,6 +306,7 @@ public:
 	 * @param	bNotifySelectionChange		If true the Editor is notified of the selection change; if false, the Editor will not be notified
 	 * @param	bRedrawViewports			If true the viewports will be redrawn; if false, they will not
 	 */
+	UFUNCTION(BlueprintCallable, Category = DataLayers)
 	virtual bool UpdateAllActorsVisibility(const bool bNotifySelectionChange, const bool bRedrawViewports);
 
 
@@ -350,7 +319,8 @@ public:
 	 *	@param	DataLayer			The DataLayer to find actors for.
 	 *	@param	InOutActors			The list to append the found actors to.
 	 */
-	virtual void AppendActorsFromDataLayer(const UDataLayer* DataLayer, TArray<AActor*>& InOutActors) const;
+	UFUNCTION(BlueprintCallable, Category = DataLayers)
+	virtual void AppendActorsFromDataLayer(UDataLayer* DataLayer, TArray<AActor*>& InOutActors) const;
 
 	/**
 	 *	Appends all the actors associated with the specified DataLayer.
@@ -359,16 +329,7 @@ public:
 	 *	@param	InOutActors			The list to append the found actors to.
 	 *  @param	Filter	[optional]	Actor that don't pass the specified filter restrictions won't be selected.
 	 */
-	virtual void AppendActorsFromDataLayer(const UDataLayer* DataLayer, TArray<AActor*>& InOutActors, const TSharedPtr<FActorFilter>& Filter) const;
-	
-	/**
-	 *	Appends all the actors associated with the specified DataLayer.
-	 *
-	 *	@param	DataLayer			The DataLayer to find actors for.
-	 *	@param	InOutActors			The list to append the found actors to.
-	 *  @param	Filter	[optional]	Actor that don't pass the specified filter restrictions won't be selected.
-	 */
-	virtual void AppendActorsFromDataLayer(const UDataLayer* DataLayer, TArray<TWeakObjectPtr<AActor>>& InOutActors, const TSharedPtr<FActorFilter>& Filter = TSharedPtr<FActorFilter>(nullptr)) const;
+	virtual void AppendActorsFromDataLayer(UDataLayer* DataLayer, TArray<AActor*>& InOutActors, const TSharedPtr<FActorFilter>& Filter) const;
 	
 	/**
 	 *	Appends all the actors associated with ANY of the specified DataLayers.
@@ -376,16 +337,8 @@ public:
 	 *	@param	DataLayers			The DataLayers to find actors for.
 	 *	@param	InOutActors			The list to append the found actors to.
 	 */
-	virtual void AppendActorsFromDataLayers(const TArray<const UDataLayer*>& DataLayers, TArray<AActor*>& InOutActors) const;
-
-	/**
-	 *	Appends all the actors associated with ANY of the specified DataLayers.
-	 *
-	 *	@param	DataLayers			The DataLayers to find actors for.
-	 *	@param	InOutActors			The list to append the found actors to.
-	 *  @param	Filter	[optional]	Actor that don't pass the specified filter restrictions won't be selected.
-	 */
-	virtual void AppendActorsFromDataLayers(const TArray<const UDataLayer*>& DataLayers, TArray<AActor*>& InOutActors, const TSharedPtr<FActorFilter>& Filter) const;
+	UFUNCTION(BlueprintCallable, Category = DataLayers)
+	virtual void AppendActorsFromDataLayers(const TArray<UDataLayer*>& DataLayers, TArray<AActor*>& InOutActors) const;
 
 	/**
 	 *	Appends all the actors associated with ANY of the specified DataLayers.
@@ -394,7 +347,7 @@ public:
 	 *	@param	InOutActors			The list to append the found actors to.
 	 *  @param	Filter	[optional]	Actor that don't pass the specified filter restrictions won't be selected.
 	 */
-	virtual void AppendActorsFromDataLayers(const TArray<const UDataLayer*>& DataLayers, TArray<TWeakObjectPtr<AActor>>& InOutActors, const TSharedPtr<FActorFilter>& Filter = TSharedPtr<FActorFilter>(nullptr)) const;
+	virtual void AppendActorsFromDataLayers(const TArray<UDataLayer*>& DataLayers, TArray<AActor*>& InOutActors, const TSharedPtr<FActorFilter>& Filter) const;
 
 	/**
 	 *	Gets all the actors associated with the specified DataLayer. Analog to AppendActorsFromDataLayer but it returns rather than appends the actors.
@@ -402,7 +355,8 @@ public:
 	 *	@param	DataLayer			The DataLayer to find actors for.
 	 *	@return						The list to assign the found actors to.
 	 */
-	virtual TArray<AActor*> GetActorsFromDataLayer(const UDataLayer* DataLayer) const;
+	UFUNCTION(BlueprintCallable, Category = DataLayers)
+	virtual TArray<AActor*> GetActorsFromDataLayer(UDataLayer* DataLayer) const;
 
 	/**
 	 *	Gets all the actors associated with the specified DataLayer. Analog to AppendActorsFromDataLayer but it returns rather than appends the actors.
@@ -411,7 +365,7 @@ public:
 	 *  @param	Filter	[optional]	Actor that don't pass the specified filter restrictions won't be selected.
 	 *	@return						The list to assign the found actors to.
 	 */
-	virtual TArray<AActor*> GetActorsFromDataLayer(const UDataLayer* DataLayer, const TSharedPtr<FActorFilter>& Filter) const;
+	virtual TArray<AActor*> GetActorsFromDataLayer(UDataLayer* DataLayer, const TSharedPtr<FActorFilter>& Filter) const;
 
 	/**
 	 *	Gets all the actors associated with ANY of the specified DataLayers. Analog to AppendActorsFromDataLayers but it returns rather than appends the actors.
@@ -419,7 +373,8 @@ public:
 	 *	@param	DataLayers			The DataLayers to find actors for.
 	 *	@return						The list to assign the found actors to.
 	 */
-	virtual TArray<AActor*> GetActorsFromDataLayers(const TArray<const UDataLayer*>& DataLayers) const;
+	UFUNCTION(BlueprintCallable, Category = DataLayers)
+	virtual TArray<AActor*> GetActorsFromDataLayers(const TArray<UDataLayer*>& DataLayers) const;
 
 	/**
 	 *	Gets all the actors associated with ANY of the specified DataLayers. Analog to AppendActorsFromDataLayers but it returns rather than appends the actors.
@@ -428,7 +383,7 @@ public:
 	 *  @param	Filter	[optional]	Actor that don't pass the specified filter restrictions won't be selected.
 	 *	@return						The list to assign the found actors to.
 	 */
-	virtual TArray<AActor*> GetActorsFromDataLayers(const TArray<const UDataLayer*>& DataLayers, const TSharedPtr<FActorFilter>& Filter) const;
+	virtual TArray<AActor*> GetActorsFromDataLayers(const TArray<UDataLayer*>& DataLayers, const TSharedPtr<FActorFilter>& Filter) const;
 
 	/**
 	 * Changes the DataLayer's visibility to the provided state
@@ -436,6 +391,7 @@ public:
 	 * @param	DataLayer	The DataLayer to affect.
 	 * @param	bIsVisible	If true the DataLayer will be visible; if false, the DataLayer will not be visible.
 	 */
+	UFUNCTION(BlueprintCallable, Category = DataLayers)
 	virtual void SetDataLayerVisibility(UDataLayer* DataLayer, const bool bIsVisible);
 
 	/**
@@ -444,6 +400,7 @@ public:
 	 * @param	DataLayers	The DataLayers to affect
 	 * @param	bIsVisible	If true the DataLayers will be visible; if false, the DataLayers will not be visible
 	 */
+	UFUNCTION(BlueprintCallable, Category = DataLayers)
 	virtual void SetDataLayersVisibility(const TArray<UDataLayer*>& DataLayers, const bool bIsVisible);
 
 	/**
@@ -451,6 +408,7 @@ public:
 	 *
 	 * @param DataLayer	The DataLayer to affect
 	 */
+	UFUNCTION(BlueprintCallable, Category = DataLayers)
 	virtual void ToggleDataLayerVisibility(UDataLayer* DataLayer);
 
 	/**
@@ -458,6 +416,7 @@ public:
 	 *
 	 * @param	DataLayers	The DataLayers to affect
 	 */
+	UFUNCTION(BlueprintCallable, Category = DataLayers)
 	virtual void ToggleDataLayersVisibility(const TArray<UDataLayer*>& DataLayers);
 
 	/**
@@ -466,6 +425,7 @@ public:
 	 * @param	DataLayer				The DataLayer to affect.
 	 * @param	bIsDynamicallyLoaded	If true the DataLayer will affect runtime streaming and actors will only be loaded if the DataLayer is active.
 	 */
+	UFUNCTION(BlueprintCallable, Category = DataLayers)
 	virtual bool SetDataLayerIsDynamicallyLoaded(UDataLayer* DataLayer, const bool bIsDynamicallyLoaded);
 
 	/**
@@ -474,6 +434,7 @@ public:
 	 * @param	DataLayers				The DataLayers to affect
 	 * @param	bIsDynamicallyLoaded	If true the DataLayer will affect runtime streaming and actors will only be loaded if the DataLayer is active.
 	 */
+	UFUNCTION(BlueprintCallable, Category = DataLayers)
 	virtual bool SetDataLayersIsDynamicallyLoaded(const TArray<UDataLayer*>& DataLayers, const bool bIsDynamicallyLoaded);
 
 	/**
@@ -481,6 +442,7 @@ public:
 	 *
 	 * @param DataLayer	The DataLayer to affect
 	 */
+	UFUNCTION(BlueprintCallable, Category = DataLayers)
 	virtual bool ToggleDataLayerIsDynamicallyLoaded(UDataLayer* DataLayer);
 
 	/**
@@ -488,6 +450,7 @@ public:
 	 *
 	 * @param	DataLayers	The DataLayers to affect
 	 */
+	UFUNCTION(BlueprintCallable, Category = DataLayers)
 	virtual bool ToggleDataLayersIsDynamicallyLoaded(const TArray<UDataLayer*>& DataLayers);
 
 	/**
@@ -498,6 +461,7 @@ public:
 	 *											If the DataLayer is DynamicallyLoaded, the Editor loading will consider this DataLayer to load or not an Actor.
 	 *											An Actor will not be loaded in the Editor if all its DataLayers are DynamicallyLoaded and not DynamicallyLoadedInEditor.
 	 */
+	UFUNCTION(BlueprintCallable, Category = DataLayers)
 	virtual bool SetDataLayerIsDynamicallyLoadedInEditor(UDataLayer* DataLayer, const bool bIsDynamicallyLoadedInEditor);
 
 	/**
@@ -508,6 +472,7 @@ public:
 	 *											If the DataLayer is DynamicallyLoaded, the Editor loading will consider this DataLayer to load or not an Actor.
 	 *											An Actor will not be loaded in the Editor if all its DataLayers are DynamicallyLoaded and not DynamicallyLoadedInEditor.
 	 */
+	UFUNCTION(BlueprintCallable, Category = DataLayers)
 	virtual bool SetDataLayersIsDynamicallyLoadedInEditor(const TArray<UDataLayer*>& DataLayers, const bool bIsDynamicallyLoadedInEditor);
 
 	/**
@@ -515,6 +480,7 @@ public:
 	 *
 	 * @param DataLayer	The DataLayer to affect
 	 */
+	UFUNCTION(BlueprintCallable, Category = DataLayers)
 	virtual bool ToggleDataLayerIsDynamicallyLoadedInEditor(UDataLayer* DataLayer);
 
 	/**
@@ -522,11 +488,13 @@ public:
 	 *
 	 * @param	DataLayers	The DataLayers to affect
 	 */
+	UFUNCTION(BlueprintCallable, Category = DataLayers)
 	virtual bool ToggleDataLayersIsDynamicallyLoadedInEditor(const TArray<UDataLayer*>& DataLayers);
 
 	/**
 	 * Set the visibility of all DataLayers to true
 	 */
+	UFUNCTION(BlueprintCallable, Category = DataLayers)
 	virtual void MakeAllDataLayersVisible();
 
 	/**
@@ -535,7 +503,17 @@ public:
 	 * @param	DataLayerLabel	The label of the DataLayer whose UDataLayer Object is returned
 	 * @return					The UDataLayer Object of the provided DataLayer label
 	 */
-	virtual const UDataLayer* GetDataLayerFromLabel(const FName& DataLayerLabel) const;
+	UFUNCTION(BlueprintCallable, Category = DataLayers)
+	virtual UDataLayer* GetDataLayerFromLabel(const FName& DataLayerLabel) const;
+
+	/**
+	 * Gets the UDataLayer Object of the ActorDataLayer
+	 *
+	 * @param	ActorDataLayer	The ActorDataLayer whose UDataLayer Object is returned
+	 * @return					The UDataLayer Object of the provided ActorDataLayer
+	 */
+	UFUNCTION(BlueprintCallable, Category = DataLayers)
+	virtual UDataLayer* GetDataLayer(const FActorDataLayer& ActorDataLayer) const;
 
 	/**
 	 * Gets the UDataLayer Object of the DataLayer name
@@ -543,7 +521,7 @@ public:
 	 * @param	DataLayerName	The name of the DataLayer whose UDataLayer Object is returned
 	 * @return					The UDataLayer Object of the provided DataLayer name
 	 */
-	virtual const UDataLayer* GetDataLayerFromName(const FName& DataLayerName) const;
+	virtual UDataLayer* GetDataLayerFromName(const FName& DataLayerName) const;
 	
 	/**
 	 * Attempts to get the UDataLayer Object of the provided DataLayer label.
@@ -552,7 +530,7 @@ public:
 	 * @param	OutDataLayer[OUT] 	Set to the UDataLayer Object of the DataLayer label. Set to Invalid if no UDataLayer Object exists.
 	 * @return					If true a valid UDataLayer Object was found and set to OutDataLayer; if false, a valid UDataLayer object was not found and invalid set to OutDataLayer
 	 */
-	virtual bool TryGetDataLayerFromLabel(const FName& DataLayerLabel, const UDataLayer*& OutDataLayer);
+	virtual bool TryGetDataLayerFromLabel(const FName& DataLayerLabel, UDataLayer*& OutDataLayer);
 
 	/**
 	 * Gets all known DataLayers and appends them to the provided array
@@ -566,44 +544,64 @@ public:
 	 *
 	 * @return	The newly created UDataLayer Object
 	 */
-	 UDataLayer* CreateDataLayer();
+	UFUNCTION(BlueprintCallable, Category = DataLayers)
+	UDataLayer* CreateDataLayer();
 
 	/**
-	 * Gets the WorldDataLayers for this world
+	 * Deletes all of the provided DataLayers
 	 *
-	 * @param	bCreateIfNotFound	If true, created the AWorldDataLayers object if not found.
-	 * @return	The AWorldDataLayers Object
+	 * @param DataLayersToDelete	A valid list of DataLayer.
 	 */
-	class AWorldDataLayers* GetWorldDataLayers(bool bCreateIfNotFound = false);
+	UFUNCTION(BlueprintCallable, Category = DataLayers)
+	virtual void DeleteDataLayers(const TArray<UDataLayer*>& DataLayersToDelete);
 
 	/**
-	 * Gets the WorldDataLayers for this world
+	 * Deletes the provided DataLayer
 	 *
-	 * @return	The AWorldDataLayers Object
+	 * @param DataLayerToDelete		A valid DataLayer
 	 */
-	const class AWorldDataLayers* GetWorldDataLayers() const;
+	UFUNCTION(BlueprintCallable, Category = DataLayers)
+	virtual void DeleteDataLayer(UDataLayer* DataLayerToDelete);
 
 	/**
-	 * Deletes all of the provided DataLayers, disassociating all actors from them
+	 * Renames the provided DataLayer to the provided new label
 	 *
-	 * @param DataLayersToDelete	A valid list of DataLayer.
+	 * @param	DataLayer			The DataLayer to be renamed
+	 * @param	NewDataLayerLabel	The new label for the DataLayer to be renamed
 	 */
-	virtual void DeleteDataLayers(const TArray<const UDataLayer*>& DataLayersToDelete);
+	UFUNCTION(BlueprintCallable, Category = DataLayers)
+	virtual bool RenameDataLayer(UDataLayer* DataLayer, const FName& NewDataLayerLabel);
+
+private:
 
 	/**
-	 * Deletes the provided DataLayer, disassociating all actors from them
+	 *	Synchronizes an newly created Actor's DataLayers with the DataLayer system
 	 *
-	 * @param DataLayerToDelete		A valid DataLayer
+	 *	@param	Actor	The actor to initialize
 	 */
-	virtual void DeleteDataLayer(const UDataLayer* DataLayerToDelete);
+	void InitializeNewActorDataLayers(AActor* Actor);
 
 	/**
-	 * Renames the provided DataLayer to the provided new label
+	 * Find and return the selected actors.
 	 *
-	 * @param	DataLayer			The DataLayer to be renamed
-	 * @param	NewDataLayerLabel	The new label for the DataLayer to be renamed
+	 * @return				The selected AActor's as a TArray.
 	 */
-	virtual bool RenameDataLayer(UDataLayer* DataLayer, const FName& NewDataLayerLabel);
+	TArray<AActor*> GetSelectedActors() const;
+
+	/**
+	 * Gets the WorldDataLayers for this world
+	 *
+	 * @param	bCreateIfNotFound	If true, created the AWorldDataLayers object if not found.
+	 * @return	The AWorldDataLayers Object
+	 */
+	class AWorldDataLayers* GetWorldDataLayers(bool bCreateIfNotFound = false);
+
+	/**
+	 * Gets the WorldDataLayers for this world
+	 *
+	 * @return	The AWorldDataLayers Object
+	 */
+	const class AWorldDataLayers* GetWorldDataLayers() const;
 
 	/**
 	 * Get the current UWorld object.
@@ -612,23 +610,15 @@ public:
 	 */
 	UWorld* GetWorld() const; // Fallback to GWorld
 
-	/**
-	 * Delegate handler for FEditorDelegates::MapChange. It internally calls DataLayerChanged.Broadcast.
-	 **/
+	/** Delegate handler for FEditorDelegates::MapChange. It internally calls DataLayerChanged.Broadcast. */
 	void EditorMapChange();
-	
-	/**
-	 * Delegate handler for FEditorDelegates::RefreshDataLayerBrowser. It internally calls UpdateAllActorsVisibility to refresh the actors of each DataLayer.
-	 **/
+
+	/** Delegate handler for FEditorDelegates::RefreshDataLayerBrowser. It internally calls UpdateAllActorsVisibility to refresh the actors of each DataLayer. */
 	void EditorRefreshDataLayerBrowser();
 
-	/**
-	 * Delegate handler for FEditorDelegates::PostUndoRedo. It internally calls DataLayerChanged.Broadcast and UpdateAllActorsVisibility to refresh the actors of each DataLayer.
-	 **/
+	/** Delegate handler for FEditorDelegates::PostUndoRedo. It internally calls DataLayerChanged.Broadcast and UpdateAllActorsVisibility to refresh the actors of each DataLayer. */
 	void PostUndoRedo();
 
-private:
-
 	bool SetDataLayerIsDynamicallyLoadedInternal(UDataLayer* DataLayer, const bool bIsDynamicallyLoaded);
 	bool SetDataLayerIsDynamicallyLoadedInEditorInternal(UDataLayer* DataLayer, const bool bIsDynamicallyLoadedInEditor);
 	bool RefreshWorldPartitionEditorCells();
@@ -642,4 +632,6 @@ private:
 
 	/** Auxiliary class that sets the callback functions for multiple delegates */
 	TSharedPtr<class FDataLayersBroadcast> DataLayersBroadcast;
+
+	friend class FDataLayersBroadcast;
 };
diff --git a/Engine/Source/Editor/LevelEditor/Private/SLevelViewport.cpp b/Engine/Source/Editor/LevelEditor/Private/SLevelViewport.cpp
--- a/Engine/Source/Editor/LevelEditor/Private/SLevelViewport.cpp
+++ b/Engine/Source/Editor/LevelEditor/Private/SLevelViewport.cpp
@@ -2074,7 +2074,6 @@ bool SLevelViewport::IsLayerVisible( FName LayerName ) const
 /** Called when a user selects show or hide all from the DataLayers visibility menu. **/
 void SLevelViewport::OnToggleAllDataLayers(bool bVisible)
 {
-	UDataLayerEditorSubsystem* DataLayerSubsystem = UDataLayerEditorSubsystem::Get();
 	if (bVisible)
 	{
 		// Clear all hidden DataLayers
@@ -2096,7 +2095,7 @@ void SLevelViewport::OnToggleAllDataLayers(bool bVisible)
 	}
 
 	// Update actor visibility for this view
-	DataLayerSubsystem->UpdatePerViewVisibility(LevelViewportClient.Get());
+	UDataLayerEditorSubsystem::Get()->UpdatePerViewVisibility(LevelViewportClient.Get());
 	LevelViewportClient->Invalidate();
 }
 
@@ -2115,7 +2114,7 @@ void SLevelViewport::ToggleShowDataLayer(FName DataLayerName)
 
 	// Update actor visibility for this view
 	UDataLayerEditorSubsystem* DataLayerEditorSubsystem = UDataLayerEditorSubsystem::Get();
-	if (const UDataLayer* DataLayer = DataLayerEditorSubsystem->GetDataLayerFromName(DataLayerName))
+	if (UDataLayer* DataLayer = DataLayerEditorSubsystem->GetDataLayerFromName(DataLayerName))
 	{
 		DataLayerEditorSubsystem->UpdatePerViewVisibility(LevelViewportClient.Get(), DataLayer);
 		LevelViewportClient->Invalidate();
diff --git a/Engine/Source/Runtime/Engine/Private/WorldPartition/DataLayer/WorldDataLayers.cpp b/Engine/Source/Runtime/Engine/Private/WorldPartition/DataLayer/WorldDataLayers.cpp
--- a/Engine/Source/Runtime/Engine/Private/WorldPartition/DataLayer/WorldDataLayers.cpp
+++ b/Engine/Source/Runtime/Engine/Private/WorldPartition/DataLayer/WorldDataLayers.cpp
@@ -104,7 +104,7 @@ UDataLayer* AWorldDataLayers::CreateDataLayer()
 	return NewDataLayer;
 }
 
-bool AWorldDataLayers::RemoveDataLayers(const TArray<const UDataLayer*>& InDataLayers)
+bool AWorldDataLayers::RemoveDataLayers(const TArray<UDataLayer*>& InDataLayers)
 {
 	bool bIsModified = false;
 	for (const UDataLayer* DataLayer : InDataLayers)
@@ -119,7 +119,7 @@ bool AWorldDataLayers::RemoveDataLayers(const TArray<const UDataLayer*>& InDataL
 	return bIsModified;
 }
 
-bool AWorldDataLayers::RemoveDataLayer(const UDataLayer* InDataLayer)
+bool AWorldDataLayers::RemoveDataLayer(UDataLayer* InDataLayer)
 {
 	if (ContainsDataLayer(InDataLayer))
 	{
diff --git a/Engine/Source/Runtime/Engine/Public/WorldPartition/DataLayer/WorldDataLayers.h b/Engine/Source/Runtime/Engine/Public/WorldPartition/DataLayer/WorldDataLayers.h
--- a/Engine/Source/Runtime/Engine/Public/WorldPartition/DataLayer/WorldDataLayers.h
+++ b/Engine/Source/Runtime/Engine/Public/WorldPartition/DataLayer/WorldDataLayers.h
@@ -24,8 +24,8 @@ public:
 #if WITH_EDITOR
 	static AWorldDataLayers* Get(UWorld* World, bool bCreateIfNotFound);
 	UDataLayer* CreateDataLayer();
-	bool RemoveDataLayer(const UDataLayer* InDataLayer);
-	bool RemoveDataLayers(const TArray<const UDataLayer*>& InDataLayers);
+	bool RemoveDataLayer(UDataLayer* InDataLayer);
+	bool RemoveDataLayers(const TArray<UDataLayer*>& InDataLayers);
 	FName GenerateUniqueDataLayerLabel(const FName& InDataLayerLabel) const;
 #endif
 	
