Fix DataLayers module crash when exiting the editor.

diff --git a/Engine/Source/Editor/DataLayerEditor/Private/DataLayerEditorModule.cpp b/Engine/Source/Editor/DataLayerEditor/Private/DataLayerEditorModule.cpp
--- a/Engine/Source/Editor/DataLayerEditor/Private/DataLayerEditorModule.cpp
+++ b/Engine/Source/Editor/DataLayerEditor/Private/DataLayerEditorModule.cpp
@@ -12,17 +12,19 @@
 
 IMPLEMENT_MODULE(FDataLayerEditorModule, DataLayerEditor );
 
+static const FName NAME_ActorDataLayer(TEXT("ActorDataLayer"));
+
 void FDataLayerEditorModule::StartupModule()
 {
 	FPropertyEditorModule& PropertyModule = FModuleManager::LoadModuleChecked<FPropertyEditorModule>("PropertyEditor");
-	PropertyModule.RegisterCustomPropertyTypeLayout(FActorDataLayer::StaticStruct()->GetFName(), FOnGetPropertyTypeCustomizationInstance::CreateLambda([] { return MakeShared<FDataLayerPropertyTypeCustomization>(); }));
+	PropertyModule.RegisterCustomPropertyTypeLayout(NAME_ActorDataLayer, FOnGetPropertyTypeCustomizationInstance::CreateLambda([] { return MakeShared<FDataLayerPropertyTypeCustomization>(); }));
 }
 
 void FDataLayerEditorModule::ShutdownModule()
 {
 	if (FPropertyEditorModule* PropertyModule = FModuleManager::GetModulePtr<FPropertyEditorModule>("PropertyEditor"))
 	{
-		PropertyModule->UnregisterCustomPropertyTypeLayout(FActorDataLayer::StaticStruct()->GetFName());
+		PropertyModule->UnregisterCustomPropertyTypeLayout(NAME_ActorDataLayer);
 	}
 }
 
