Fixed non-unity build

diff --git a/Engine/Source/Editor/DataLayerEditor/Public/DataLayer/DataLayerEditorSubsystem.h b/Engine/Source/Editor/DataLayerEditor/Public/DataLayer/DataLayerEditorSubsystem.h
--- a/Engine/Source/Editor/DataLayerEditor/Public/DataLayer/DataLayerEditorSubsystem.h
+++ b/Engine/Source/Editor/DataLayerEditor/Public/DataLayer/DataLayerEditorSubsystem.h
@@ -5,6 +5,7 @@
 #include "CoreMinimal.h"
 #include "EditorSubsystem.h"
 #include "DataLayerAction.h"
+#include "WorldPartition/DataLayer/ActorDataLayer.h"
 #include "DataLayerEditorSubsystem.generated.h"
 
 class AActor;
