Fixing PVS 7.7 Issues: warning V1052: Declaring virtual methods in a class marked as 'final' is pointless. Consider inspecting the 'XYZ' method of the 'FXYZ' class.

diff --git a/Engine/Plugins/Experimental/MeshModelingToolset/Source/ModelingComponents/Private/OctreeDynamicMeshSceneProxy.h b/Engine/Plugins/Experimental/MeshModelingToolset/Source/ModelingComponents/Private/OctreeDynamicMeshSceneProxy.h
--- a/Engine/Plugins/Experimental/MeshModelingToolset/Source/ModelingComponents/Private/OctreeDynamicMeshSceneProxy.h
+++ b/Engine/Plugins/Experimental/MeshModelingToolset/Source/ModelingComponents/Private/OctreeDynamicMeshSceneProxy.h
@@ -64,7 +64,7 @@ public:
 
 
 
-	virtual void InitializeSingleBuffer()
+	void InitializeSingleBuffer()
 	{
 		check(RenderBufferSets.Num() == 0);
 
@@ -96,7 +96,7 @@ public:
 
 
 
-	virtual void InitializeFromDecomposition(const FArrayIndexSetsDecomposition& Decomposition)
+	void InitializeFromDecomposition(const FArrayIndexSetsDecomposition& Decomposition)
 	{
 		check(RenderBufferSets.Num() == 0);
 
@@ -134,7 +134,7 @@ public:
 
 
 
-	virtual void UpdateFromDecomposition(const FArrayIndexSetsDecomposition& Decomposition, const TArray<int32>& SetsToUpdate )
+	void UpdateFromDecomposition(const FArrayIndexSetsDecomposition& Decomposition, const TArray<int32>& SetsToUpdate )
 	{
 		// CAN WE REUSE EXISTING BUFFER SETS??
 		//   - could have timestamp for each decomposition set array...if tris don't change we only have to update vertices
diff --git a/Engine/Plugins/Experimental/MeshModelingToolset/Source/ModelingComponents/Private/SimpleDynamicMeshSceneProxy.h b/Engine/Plugins/Experimental/MeshModelingToolset/Source/ModelingComponents/Private/SimpleDynamicMeshSceneProxy.h
--- a/Engine/Plugins/Experimental/MeshModelingToolset/Source/ModelingComponents/Private/SimpleDynamicMeshSceneProxy.h
+++ b/Engine/Plugins/Experimental/MeshModelingToolset/Source/ModelingComponents/Private/SimpleDynamicMeshSceneProxy.h
@@ -46,7 +46,7 @@ public:
 
 
 
-	virtual void Initialize()
+	void Initialize()
 	{
 		// allocate buffer sets based on materials
 		check(RenderBufferSets.Num() == 0);
@@ -86,7 +86,7 @@ public:
 	/**
 	 * Initialize multiple buffers for the mesh based on the given Decomposition
 	 */
-	virtual void InitializeFromDecomposition(TUniquePtr<FMeshRenderDecomposition>& Decomposition)
+	void InitializeFromDecomposition(TUniquePtr<FMeshRenderDecomposition>& Decomposition)
 	{
 		check(RenderBufferSets.Num() == 0);
 		int32 NumSets = Decomposition->Num();
@@ -150,7 +150,7 @@ public:
 	/**
 	 * Initialize a single set of mesh buffers for the entire mesh
 	 */
-	virtual void InitializeSingleBufferSet(FMeshRenderBufferSet* RenderBuffers)
+	void InitializeSingleBufferSet(FMeshRenderBufferSet* RenderBuffers)
 	{
 		FDynamicMesh3* Mesh = ParentComponent->GetRenderMesh();
 
@@ -189,7 +189,7 @@ public:
 	/**
 	 * Initialize the mesh buffers, one per material
 	 */
-	virtual void InitializeByMaterial(TArray<FMeshRenderBufferSet*>& BufferSets)
+	void InitializeByMaterial(TArray<FMeshRenderBufferSet*>& BufferSets)
 	{
 		FDynamicMesh3* Mesh = ParentComponent->GetRenderMesh();
 		check(Mesh->HasAttributes() && Mesh->Attributes()->HasMaterialID());
@@ -338,7 +338,7 @@ public:
 	/**
 	 * Update the vertex position/normal/color buffers
 	 */
-	virtual void FastUpdateVertices(bool bPositions, bool bNormals, bool bColors, bool bUVs)
+	void FastUpdateVertices(bool bPositions, bool bNormals, bool bColors, bool bUVs)
 	{
 		// This needs to be rewritten for split-by-material buffers.
 		// Could store triangle set with each buffer, and then rebuild vtx buffer(s) as needed?
@@ -431,7 +431,7 @@ public:
 	/**
 	 * Update the vertex position/normal/color buffers
 	 */
-	virtual void FastUpdateVertices(const TArray<int32>& WhichBuffers, bool bPositions, bool bNormals, bool bColors, bool bUVs)
+	void FastUpdateVertices(const TArray<int32>& WhichBuffers, bool bPositions, bool bNormals, bool bColors, bool bUVs)
 	{
 		// This needs to be rewritten for split-by-material buffers.
 		// Could store triangle set with each buffer, and then rebuild vtx buffer(s) as needed?
@@ -499,7 +499,7 @@ public:
 	/**
 	 * Update index buffers inside each RenderBuffer set
 	 */
-	virtual void FastUpdateAllIndexBuffers()
+	void FastUpdateAllIndexBuffers()
 	{
 		FDynamicMesh3* Mesh = ParentComponent->GetRenderMesh();
 
diff --git a/Engine/Plugins/Online/OnlineSubsystemSteam/Source/Private/OnlineSessionAsyncServerSteam.h b/Engine/Plugins/Online/OnlineSubsystemSteam/Source/Private/OnlineSessionAsyncServerSteam.h
--- a/Engine/Plugins/Online/OnlineSubsystemSteam/Source/Private/OnlineSessionAsyncServerSteam.h
+++ b/Engine/Plugins/Online/OnlineSubsystemSteam/Source/Private/OnlineSessionAsyncServerSteam.h
@@ -208,7 +208,7 @@ public:
 
 	}
 
-	virtual ~FPendingSearchResultSteam()
+	~FPendingSearchResultSteam()
 	{
 		CancelQuery();
 	}
diff --git a/Engine/Source/Developer/StandaloneRenderer/Private/Windows/D3D/SlateD3DShaders.cpp b/Engine/Source/Developer/StandaloneRenderer/Private/Windows/D3D/SlateD3DShaders.cpp
--- a/Engine/Source/Developer/StandaloneRenderer/Private/Windows/D3D/SlateD3DShaders.cpp
+++ b/Engine/Source/Developer/StandaloneRenderer/Private/Windows/D3D/SlateD3DShaders.cpp
@@ -69,7 +69,7 @@ static pD3DReflect GetD3DReflectFunc()
 class StandaloneD3DIncluder final : public ID3DInclude
 {
 	public:
-		virtual ~StandaloneD3DIncluder() 
+		~StandaloneD3DIncluder() 
 		{
 		}
 
diff --git a/Engine/Source/Editor/DataLayerEditor/Public/DataLayer/DataLayerEditorSubsystem.h b/Engine/Source/Editor/DataLayerEditor/Public/DataLayer/DataLayerEditorSubsystem.h
--- a/Engine/Source/Editor/DataLayerEditor/Public/DataLayer/DataLayerEditorSubsystem.h
+++ b/Engine/Source/Editor/DataLayerEditor/Public/DataLayer/DataLayerEditorSubsystem.h
@@ -53,11 +53,11 @@ public:
 	 * Reset  : A large amount of changes have occurred to a number of DataLayers.
 	 */
 	DECLARE_EVENT_ThreeParams(UDataLayerEditorSubsystem, FOnDataLayerChanged, const EDataLayerAction /*Action*/, const TWeakObjectPtr<const UDataLayer>& /*ChangedDataLayer*/, const FName& /*ChangedProperty*/);
-	virtual FOnDataLayerChanged& OnDataLayerChanged() { return DataLayerChanged; }
+	FOnDataLayerChanged& OnDataLayerChanged() { return DataLayerChanged; }
 
 	/** Broadcasts whenever one or more Actors changed DataLayers*/
 	DECLARE_EVENT_OneParam(UDataLayerEditorSubsystem, FOnActorDataLayersChanged, const TWeakObjectPtr<AActor>& /*ChangedActor*/);
-	virtual FOnActorDataLayersChanged& OnActorDataLayersChanged() { return ActorDataLayersChanged; }
+	FOnActorDataLayersChanged& OnActorDataLayersChanged() { return ActorDataLayersChanged; }
 	
 	///////////////////////////////////////////////////////////////////////////////////////////////////////////////////////
 	// Operations on an individual actor.
@@ -68,7 +68,7 @@ public:
 	 *	@param	Actor	The actor to validate
 	 */
 	UFUNCTION(BlueprintCallable, Category = DataLayers)
-	virtual bool IsActorValidForDataLayer(AActor* Actor);
+	bool IsActorValidForDataLayer(AActor* Actor);
 
 	/**
 	 * Adds the actor to the DataLayer.
@@ -78,7 +78,7 @@ public:
 	 * @return				true if the actor was added.  false is returned if the actor already belongs to the DataLayer.
 	 */
 	UFUNCTION(BlueprintCallable, Category = DataLayers)
-	virtual bool AddActorToDataLayer(AActor* Actor, UDataLayer* DataLayer);
+	bool AddActorToDataLayer(AActor* Actor, UDataLayer* DataLayer);
 
 	/**
 	 * Adds the provided actor to the DataLayers.
@@ -88,7 +88,7 @@ public:
 	 * @return				true if the actor was added to at least one of the provided DataLayers.
 	 */
 	UFUNCTION(BlueprintCallable, Category = DataLayers)
-	virtual bool AddActorToDataLayers(AActor* Actor, const TArray<UDataLayer*>& DataLayers);
+	bool AddActorToDataLayers(AActor* Actor, const TArray<UDataLayer*>& DataLayers);
 
 	/**
 	 * Removes an actor from the specified DataLayer.
@@ -98,7 +98,7 @@ public:
 	 * @return					true if the actor was removed from the DataLayer.  false is returned if the actor already belonged to the DataLayer.
 	 */
 	UFUNCTION(BlueprintCallable, Category = DataLayers)
-	virtual bool RemoveActorFromDataLayer(AActor* Actor, UDataLayer* DataLayerToRemove);
+	bool RemoveActorFromDataLayer(AActor* Actor, UDataLayer* DataLayerToRemove);
 
 	/**
 	 * Removes the provided actor from the DataLayers.
@@ -108,7 +108,7 @@ public:
 	 * @return				true if the actor was removed from at least one of the provided DataLayers.
 	 */
 	UFUNCTION(BlueprintCallable, Category = DataLayers)
-	virtual bool RemoveActorFromDataLayers(AActor* Actor, const TArray<UDataLayer*>& DataLayers);
+	bool RemoveActorFromDataLayers(AActor* Actor, const TArray<UDataLayer*>& DataLayers);
 
 
 	/////////////////////////////////////////////////
@@ -122,7 +122,7 @@ public:
 	 * @return				true if at least one actor was added to the DataLayer.  false is returned if all the actors already belonged to the DataLayer.
 	 */
 	UFUNCTION(BlueprintCallable, Category = DataLayers)
-	virtual bool AddActorsToDataLayer(const TArray<AActor*>& Actors, UDataLayer* DataLayer);
+	bool AddActorsToDataLayer(const TArray<AActor*>& Actors, UDataLayer* DataLayer);
 
 	/**
 	 * Add the actors to the DataLayers
@@ -142,7 +142,7 @@ public:
 	 * @return					true if at least one actor was removed from the DataLayer.  false is returned if all the actors already belonged to the DataLayer.
 	 */
 	UFUNCTION(BlueprintCallable, Category = DataLayers)
-	virtual bool RemoveActorsFromDataLayer(const TArray<AActor*>& Actors, UDataLayer* DataLayer);
+	bool RemoveActorsFromDataLayer(const TArray<AActor*>& Actors, UDataLayer* DataLayer);
 
 	/**
 	 * Remove the actors to the DataLayers
@@ -152,7 +152,7 @@ public:
 	 * @return				true if at least one actor was removed from at least one DataLayer.  false is returned if none of the actors belonged to any of the specified DataLayers.
 	 */
 	UFUNCTION(BlueprintCallable, Category = DataLayers)
-	virtual bool RemoveActorsFromDataLayers(const TArray<AActor*>& Actors, const TArray<UDataLayer*>& DataLayers);
+	bool RemoveActorsFromDataLayers(const TArray<AActor*>& Actors, const TArray<UDataLayer*>& DataLayers);
 
 	/////////////////////////////////////////////////
 	// Operations on selected actors.
@@ -164,7 +164,7 @@ public:
 	 * @return				true if at least one actor was added.  false is returned if all selected actors already belong to the DataLayer.
 	 */
 	UFUNCTION(BlueprintCallable, Category = DataLayers)
-	virtual bool AddSelectedActorsToDataLayer(UDataLayer* DataLayer);
+	bool AddSelectedActorsToDataLayer(UDataLayer* DataLayer);
 
 	/**
 	 * Adds selected actors to the DataLayers.
@@ -173,7 +173,7 @@ public:
 	 * @return				true if at least one actor was added.  false is returned if all selected actors already belong to the DataLayers.
 	 */
 	UFUNCTION(BlueprintCallable, Category = DataLayers)
-	virtual bool AddSelectedActorsToDataLayers(const TArray<UDataLayer*>& DataLayers);
+	bool AddSelectedActorsToDataLayers(const TArray<UDataLayer*>& DataLayers);
 
 	/**
 	 * Removes the selected actors from the DataLayer.
@@ -182,7 +182,7 @@ public:
 	 * @return				true if at least one actor was added.  false is returned if all selected actors already belong to the DataLayer.
 	 */
 	UFUNCTION(BlueprintCallable, Category = DataLayers)
-	virtual bool RemoveSelectedActorsFromDataLayer(UDataLayer* DataLayer);
+	bool RemoveSelectedActorsFromDataLayer(UDataLayer* DataLayer);
 
 	/**
 	 * Removes selected actors from the DataLayers.
@@ -191,7 +191,7 @@ public:
 	 * @return				true if at least one actor was removed.
 	 */
 	UFUNCTION(BlueprintCallable, Category = DataLayers)
-	virtual bool RemoveSelectedActorsFromDataLayers(const TArray<UDataLayer*>& DataLayers);
+	bool RemoveSelectedActorsFromDataLayers(const TArray<UDataLayer*>& DataLayers);
 
 
 	/////////////////////////////////////////////////
@@ -207,7 +207,7 @@ public:
 	 * @return									true if at least one actor was selected/deselected.
 	 */
 	UFUNCTION(BlueprintCallable, Category = DataLayers)
-	virtual bool SelectActorsInDataLayer(UDataLayer* DataLayer, const bool bSelect, const bool bNotify, const bool bSelectEvenIfHidden = false);
+	bool SelectActorsInDataLayer(UDataLayer* DataLayer, const bool bSelect, const bool bNotify, const bool bSelectEvenIfHidden = false);
 
 	/**
 	 * Selects/de-selects actors belonging to the DataLayer.
@@ -219,7 +219,7 @@ public:
 	 * @param	Filter	[optional]				Actor that don't pass the specified filter restrictions won't be selected.
 	 * @return									true if at least one actor was selected/deselected.
 	 */
-	virtual bool SelectActorsInDataLayer(UDataLayer* DataLayer, const bool bSelect, const bool bNotify, const bool bSelectEvenIfHidden, const TSharedPtr<FActorFilter>& Filter);
+	bool SelectActorsInDataLayer(UDataLayer* DataLayer, const bool bSelect, const bool bNotify, const bool bSelectEvenIfHidden, const TSharedPtr<FActorFilter>& Filter);
 
 	/**
 	 * Selects/de-selects actors belonging to the DataLayers.
@@ -231,7 +231,7 @@ public:
 	 * @return									true if at least one actor was selected/deselected.
 	 */
 	UFUNCTION(BlueprintCallable, Category = DataLayers)
-	virtual bool SelectActorsInDataLayers(const TArray<UDataLayer*>& DataLayers, const bool bSelect, const bool bNotify, const bool bSelectEvenIfHidden = false);
+	bool SelectActorsInDataLayers(const TArray<UDataLayer*>& DataLayers, const bool bSelect, const bool bNotify, const bool bSelectEvenIfHidden = false);
 
 	/**
 	 * Selects/de-selects actors belonging to the DataLayers.
@@ -243,7 +243,7 @@ public:
 	 * @param	Filter	[optional]				Actor that don't pass the specified filter restrictions won't be selected.
 	 * @return									true if at least one actor was selected/deselected.
 	 */
-	virtual bool SelectActorsInDataLayers(const TArray<UDataLayer*>& DataLayers, const bool bSelect, const bool bNotify, const bool bSelectEvenIfHidden, const TSharedPtr<FActorFilter>& Filter);
+	bool SelectActorsInDataLayers(const TArray<UDataLayer*>& DataLayers, const bool bSelect, const bool bNotify, const bool bSelectEvenIfHidden, const TSharedPtr<FActorFilter>& Filter);
 
 
 	/////////////////////////////////////////////////
@@ -255,7 +255,7 @@ public:
 	 * @param DataLayerThatChanged  If one DataLayer was changed (toggled in view pop-up, etc), then we only need to modify actors that use that DataLayer.
 	 */
 	UFUNCTION(BlueprintCallable, Category = DataLayers)
-	virtual void UpdateAllViewVisibility(UDataLayer* DataLayerThatChanged);
+	void UpdateAllViewVisibility(UDataLayer* DataLayerThatChanged);
 
 	/**
 	 * Updates the per-view visibility for all actors for the given view
@@ -263,7 +263,7 @@ public:
 	 * @param ViewportClient				The viewport client to update visibility on
 	 * @param DataLayerThatChanged [optional]	If one DataLayer was changed (toggled in view pop-up, etc), then we only need to modify actors that use that DataLayer
 	 */
-	virtual void UpdatePerViewVisibility(FLevelEditorViewportClient* ViewportClient, UDataLayer* DataLayerThatChanged = nullptr);
+	void UpdatePerViewVisibility(FLevelEditorViewportClient* ViewportClient, UDataLayer* DataLayerThatChanged = nullptr);
 
 	/**
 	 * Updates per-view visibility for the given actor in the given view
@@ -272,7 +272,7 @@ public:
 	 * @param Actor								Actor to update
 	 * @param bReregisterIfDirty [optional]		If true, the actor will reregister itself to give the rendering thread updated information
 	 */
-	virtual void UpdateActorViewVisibility(FLevelEditorViewportClient* ViewportClient, AActor* Actor, const bool bReregisterIfDirty = true);
+	void UpdateActorViewVisibility(FLevelEditorViewportClient* ViewportClient, AActor* Actor, const bool bReregisterIfDirty = true);
 
 	/**
 	 * Updates per-view visibility for the given actor for all views
@@ -280,14 +280,14 @@ public:
 	 * @param Actor		Actor to update
 	 */
 	UFUNCTION(BlueprintCallable, Category = DataLayers)
-	virtual void UpdateActorAllViewsVisibility(AActor* Actor);
+	void UpdateActorAllViewsVisibility(AActor* Actor);
 
 	/**
 	 * Removes the corresponding visibility bit from all actors (slides the later bits down 1)
 	 *
 	 * @param ViewportClient	The viewport client to update visibility on
 	 */
-	virtual void RemoveViewFromActorViewVisibility(FLevelEditorViewportClient* ViewportClient);
+	void RemoveViewFromActorViewVisibility(FLevelEditorViewportClient* ViewportClient);
 
 	/**
 	 * Updates the provided actors visibility in the viewports
@@ -299,7 +299,7 @@ public:
 	 * @param	bRedrawViewports			If true the viewports will be redrawn; if false, they will not
 	 */
 	UFUNCTION(BlueprintCallable, Category = DataLayers)
-	virtual bool UpdateActorVisibility(AActor* Actor, bool& bOutSelectionChanged, bool& bOutActorModified, const bool bNotifySelectionChange, const bool bRedrawViewports);
+	bool UpdateActorVisibility(AActor* Actor, bool& bOutSelectionChanged, bool& bOutActorModified, const bool bNotifySelectionChange, const bool bRedrawViewports);
 
 	/**
 	 * Updates the visibility of all actors in the viewports
@@ -308,7 +308,7 @@ public:
 	 * @param	bRedrawViewports			If true the viewports will be redrawn; if false, they will not
 	 */
 	UFUNCTION(BlueprintCallable, Category = DataLayers)
-	virtual bool UpdateAllActorsVisibility(const bool bNotifySelectionChange, const bool bRedrawViewports);
+	bool UpdateAllActorsVisibility(const bool bNotifySelectionChange, const bool bRedrawViewports);
 
 
 	/////////////////////////////////////////////////
@@ -321,7 +321,7 @@ public:
 	 *	@param	InOutActors			The list to append the found actors to.
 	 */
 	UFUNCTION(BlueprintCallable, Category = DataLayers)
-	virtual void AppendActorsFromDataLayer(UDataLayer* DataLayer, TArray<AActor*>& InOutActors) const;
+	void AppendActorsFromDataLayer(UDataLayer* DataLayer, TArray<AActor*>& InOutActors) const;
 
 	/**
 	 *	Appends all the actors associated with the specified DataLayer.
@@ -330,7 +330,7 @@ public:
 	 *	@param	InOutActors			The list to append the found actors to.
 	 *  @param	Filter	[optional]	Actor that don't pass the specified filter restrictions won't be selected.
 	 */
-	virtual void AppendActorsFromDataLayer(UDataLayer* DataLayer, TArray<AActor*>& InOutActors, const TSharedPtr<FActorFilter>& Filter) const;
+	void AppendActorsFromDataLayer(UDataLayer* DataLayer, TArray<AActor*>& InOutActors, const TSharedPtr<FActorFilter>& Filter) const;
 	
 	/**
 	 *	Appends all the actors associated with ANY of the specified DataLayers.
@@ -339,7 +339,7 @@ public:
 	 *	@param	InOutActors			The list to append the found actors to.
 	 */
 	UFUNCTION(BlueprintCallable, Category = DataLayers)
-	virtual void AppendActorsFromDataLayers(const TArray<UDataLayer*>& DataLayers, TArray<AActor*>& InOutActors) const;
+	void AppendActorsFromDataLayers(const TArray<UDataLayer*>& DataLayers, TArray<AActor*>& InOutActors) const;
 
 	/**
 	 *	Appends all the actors associated with ANY of the specified DataLayers.
@@ -348,7 +348,7 @@ public:
 	 *	@param	InOutActors			The list to append the found actors to.
 	 *  @param	Filter	[optional]	Actor that don't pass the specified filter restrictions won't be selected.
 	 */
-	virtual void AppendActorsFromDataLayers(const TArray<UDataLayer*>& DataLayers, TArray<AActor*>& InOutActors, const TSharedPtr<FActorFilter>& Filter) const;
+	void AppendActorsFromDataLayers(const TArray<UDataLayer*>& DataLayers, TArray<AActor*>& InOutActors, const TSharedPtr<FActorFilter>& Filter) const;
 
 	/**
 	 *	Gets all the actors associated with the specified DataLayer. Analog to AppendActorsFromDataLayer but it returns rather than appends the actors.
@@ -357,7 +357,7 @@ public:
 	 *	@return						The list to assign the found actors to.
 	 */
 	UFUNCTION(BlueprintCallable, Category = DataLayers)
-	virtual TArray<AActor*> GetActorsFromDataLayer(UDataLayer* DataLayer) const;
+	TArray<AActor*> GetActorsFromDataLayer(UDataLayer* DataLayer) const;
 
 	/**
 	 *	Gets all the actors associated with the specified DataLayer. Analog to AppendActorsFromDataLayer but it returns rather than appends the actors.
@@ -366,7 +366,7 @@ public:
 	 *  @param	Filter	[optional]	Actor that don't pass the specified filter restrictions won't be selected.
 	 *	@return						The list to assign the found actors to.
 	 */
-	virtual TArray<AActor*> GetActorsFromDataLayer(UDataLayer* DataLayer, const TSharedPtr<FActorFilter>& Filter) const;
+	TArray<AActor*> GetActorsFromDataLayer(UDataLayer* DataLayer, const TSharedPtr<FActorFilter>& Filter) const;
 
 	/**
 	 *	Gets all the actors associated with ANY of the specified DataLayers. Analog to AppendActorsFromDataLayers but it returns rather than appends the actors.
@@ -375,7 +375,7 @@ public:
 	 *	@return						The list to assign the found actors to.
 	 */
 	UFUNCTION(BlueprintCallable, Category = DataLayers)
-	virtual TArray<AActor*> GetActorsFromDataLayers(const TArray<UDataLayer*>& DataLayers) const;
+	TArray<AActor*> GetActorsFromDataLayers(const TArray<UDataLayer*>& DataLayers) const;
 
 	/**
 	 *	Gets all the actors associated with ANY of the specified DataLayers. Analog to AppendActorsFromDataLayers but it returns rather than appends the actors.
@@ -384,7 +384,7 @@ public:
 	 *  @param	Filter	[optional]	Actor that don't pass the specified filter restrictions won't be selected.
 	 *	@return						The list to assign the found actors to.
 	 */
-	virtual TArray<AActor*> GetActorsFromDataLayers(const TArray<UDataLayer*>& DataLayers, const TSharedPtr<FActorFilter>& Filter) const;
+	TArray<AActor*> GetActorsFromDataLayers(const TArray<UDataLayer*>& DataLayers, const TSharedPtr<FActorFilter>& Filter) const;
 
 	/**
 	 * Changes the DataLayer's visibility to the provided state
@@ -393,7 +393,7 @@ public:
 	 * @param	bIsVisible	If true the DataLayer will be visible; if false, the DataLayer will not be visible.
 	 */
 	UFUNCTION(BlueprintCallable, Category = DataLayers)
-	virtual void SetDataLayerVisibility(UDataLayer* DataLayer, const bool bIsVisible);
+	void SetDataLayerVisibility(UDataLayer* DataLayer, const bool bIsVisible);
 
 	/**
 	 * Changes visibility of the DataLayers to the provided state
@@ -402,7 +402,7 @@ public:
 	 * @param	bIsVisible	If true the DataLayers will be visible; if false, the DataLayers will not be visible
 	 */
 	UFUNCTION(BlueprintCallable, Category = DataLayers)
-	virtual void SetDataLayersVisibility(const TArray<UDataLayer*>& DataLayers, const bool bIsVisible);
+	void SetDataLayersVisibility(const TArray<UDataLayer*>& DataLayers, const bool bIsVisible);
 
 	/**
 	 * Toggles the DataLayer's visibility
@@ -410,7 +410,7 @@ public:
 	 * @param DataLayer	The DataLayer to affect
 	 */
 	UFUNCTION(BlueprintCallable, Category = DataLayers)
-	virtual void ToggleDataLayerVisibility(UDataLayer* DataLayer);
+	void ToggleDataLayerVisibility(UDataLayer* DataLayer);
 
 	/**
 	 * Toggles the visibility of all of the DataLayers
@@ -418,7 +418,7 @@ public:
 	 * @param	DataLayers	The DataLayers to affect
 	 */
 	UFUNCTION(BlueprintCallable, Category = DataLayers)
-	virtual void ToggleDataLayersVisibility(const TArray<UDataLayer*>& DataLayers);
+	void ToggleDataLayersVisibility(const TArray<UDataLayer*>& DataLayers);
 
 	/**
 	 * Changes the DataLayer's IsDynamicallyLoaded flag to the provided state
@@ -427,7 +427,7 @@ public:
 	 * @param	bIsDynamicallyLoaded	If true the DataLayer will affect runtime streaming and actors will only be loaded if the DataLayer is active.
 	 */
 	UFUNCTION(BlueprintCallable, Category = DataLayers)
-	virtual bool SetDataLayerIsDynamicallyLoaded(UDataLayer* DataLayer, const bool bIsDynamicallyLoaded);
+	bool SetDataLayerIsDynamicallyLoaded(UDataLayer* DataLayer, const bool bIsDynamicallyLoaded);
 
 	/**
 	 * Changes the IsDynamicallyLoaded flag of the DataLayers to the provided state
@@ -436,7 +436,7 @@ public:
 	 * @param	bIsDynamicallyLoaded	If true the DataLayer will affect runtime streaming and actors will only be loaded if the DataLayer is active.
 	 */
 	UFUNCTION(BlueprintCallable, Category = DataLayers)
-	virtual bool SetDataLayersIsDynamicallyLoaded(const TArray<UDataLayer*>& DataLayers, const bool bIsDynamicallyLoaded);
+	bool SetDataLayersIsDynamicallyLoaded(const TArray<UDataLayer*>& DataLayers, const bool bIsDynamicallyLoaded);
 
 	/**
 	 * Toggles the DataLayer's IsDynamicallyLoaded flag
@@ -444,7 +444,7 @@ public:
 	 * @param DataLayer	The DataLayer to affect
 	 */
 	UFUNCTION(BlueprintCallable, Category = DataLayers)
-	virtual bool ToggleDataLayerIsDynamicallyLoaded(UDataLayer* DataLayer);
+	bool ToggleDataLayerIsDynamicallyLoaded(UDataLayer* DataLayer);
 
 	/**
 	 * Toggles the IsDynamicallyLoaded flag of all of the DataLayers
@@ -452,7 +452,7 @@ public:
 	 * @param	DataLayers	The DataLayers to affect
 	 */
 	UFUNCTION(BlueprintCallable, Category = DataLayers)
-	virtual bool ToggleDataLayersIsDynamicallyLoaded(const TArray<UDataLayer*>& DataLayers);
+	bool ToggleDataLayersIsDynamicallyLoaded(const TArray<UDataLayer*>& DataLayers);
 
 	/**
 	 * Changes the DataLayer's IsDynamicallyLoadedInEditor flag to the provided state
@@ -463,7 +463,7 @@ public:
 	 *											An Actor will not be loaded in the Editor if all its DataLayers are DynamicallyLoaded and not DynamicallyLoadedInEditor.
 	 */
 	UFUNCTION(BlueprintCallable, Category = DataLayers)
-	virtual bool SetDataLayerIsDynamicallyLoadedInEditor(UDataLayer* DataLayer, const bool bIsDynamicallyLoadedInEditor);
+	bool SetDataLayerIsDynamicallyLoadedInEditor(UDataLayer* DataLayer, const bool bIsDynamicallyLoadedInEditor);
 
 	/**
 	 * Changes the IsDynamicallyLoadedInEditor flag of the DataLayers to the provided state
@@ -474,7 +474,7 @@ public:
 	 *											An Actor will not be loaded in the Editor if all its DataLayers are DynamicallyLoaded and not DynamicallyLoadedInEditor.
 	 */
 	UFUNCTION(BlueprintCallable, Category = DataLayers)
-	virtual bool SetDataLayersIsDynamicallyLoadedInEditor(const TArray<UDataLayer*>& DataLayers, const bool bIsDynamicallyLoadedInEditor);
+	bool SetDataLayersIsDynamicallyLoadedInEditor(const TArray<UDataLayer*>& DataLayers, const bool bIsDynamicallyLoadedInEditor);
 
 	/**
 	 * Toggles the DataLayer's IsDynamicallyLoadedInEditor flag
@@ -482,7 +482,7 @@ public:
 	 * @param DataLayer	The DataLayer to affect
 	 */
 	UFUNCTION(BlueprintCallable, Category = DataLayers)
-	virtual bool ToggleDataLayerIsDynamicallyLoadedInEditor(UDataLayer* DataLayer);
+	bool ToggleDataLayerIsDynamicallyLoadedInEditor(UDataLayer* DataLayer);
 
 	/**
 	 * Toggles the IsDynamicallyLoadedInEditor flag of all of the DataLayers
@@ -490,13 +490,13 @@ public:
 	 * @param	DataLayers	The DataLayers to affect
 	 */
 	UFUNCTION(BlueprintCallable, Category = DataLayers)
-	virtual bool ToggleDataLayersIsDynamicallyLoadedInEditor(const TArray<UDataLayer*>& DataLayers);
+	bool ToggleDataLayersIsDynamicallyLoadedInEditor(const TArray<UDataLayer*>& DataLayers);
 
 	/**
 	 * Set the visibility of all DataLayers to true
 	 */
 	UFUNCTION(BlueprintCallable, Category = DataLayers)
-	virtual void MakeAllDataLayersVisible();
+	void MakeAllDataLayersVisible();
 
 	/**
 	 * Gets the UDataLayer Object of the DataLayer label
@@ -505,7 +505,7 @@ public:
 	 * @return					The UDataLayer Object of the provided DataLayer label
 	 */
 	UFUNCTION(BlueprintCallable, Category = DataLayers)
-	virtual UDataLayer* GetDataLayerFromLabel(const FName& DataLayerLabel) const;
+	UDataLayer* GetDataLayerFromLabel(const FName& DataLayerLabel) const;
 
 	/**
 	 * Gets the UDataLayer Object of the ActorDataLayer
@@ -514,7 +514,7 @@ public:
 	 * @return					The UDataLayer Object of the provided ActorDataLayer
 	 */
 	UFUNCTION(BlueprintCallable, Category = DataLayers)
-	virtual UDataLayer* GetDataLayer(const FActorDataLayer& ActorDataLayer) const;
+	UDataLayer* GetDataLayer(const FActorDataLayer& ActorDataLayer) const;
 
 	/**
 	 * Gets the UDataLayer Object of the DataLayer name
@@ -522,7 +522,7 @@ public:
 	 * @param	DataLayerName	The name of the DataLayer whose UDataLayer Object is returned
 	 * @return					The UDataLayer Object of the provided DataLayer name
 	 */
-	virtual UDataLayer* GetDataLayerFromName(const FName& DataLayerName) const;
+	UDataLayer* GetDataLayerFromName(const FName& DataLayerName) const;
 	
 	/**
 	 * Attempts to get the UDataLayer Object of the provided DataLayer label.
@@ -531,14 +531,14 @@ public:
 	 * @param	OutDataLayer[OUT] 	Set to the UDataLayer Object of the DataLayer label. Set to Invalid if no UDataLayer Object exists.
 	 * @return					If true a valid UDataLayer Object was found and set to OutDataLayer; if false, a valid UDataLayer object was not found and invalid set to OutDataLayer
 	 */
-	virtual bool TryGetDataLayerFromLabel(const FName& DataLayerLabel, UDataLayer*& OutDataLayer);
+	bool TryGetDataLayerFromLabel(const FName& DataLayerLabel, UDataLayer*& OutDataLayer);
 
 	/**
 	 * Gets all known DataLayers and appends them to the provided array
 	 *
 	 * @param OutDataLayers[OUT] Output array to store all known DataLayers
 	 */
-	virtual void AddAllDataLayersTo(TArray<TWeakObjectPtr<UDataLayer>>& OutDataLayers) const;
+	void AddAllDataLayersTo(TArray<TWeakObjectPtr<UDataLayer>>& OutDataLayers) const;
 
 	/**
 	 * Creates a UDataLayer Object
@@ -554,7 +554,7 @@ public:
 	 * @param DataLayersToDelete	A valid list of DataLayer.
 	 */
 	UFUNCTION(BlueprintCallable, Category = DataLayers)
-	virtual void DeleteDataLayers(const TArray<UDataLayer*>& DataLayersToDelete);
+	void DeleteDataLayers(const TArray<UDataLayer*>& DataLayersToDelete);
 
 	/**
 	 * Deletes the provided DataLayer
@@ -562,7 +562,7 @@ public:
 	 * @param DataLayerToDelete		A valid DataLayer
 	 */
 	UFUNCTION(BlueprintCallable, Category = DataLayers)
-	virtual void DeleteDataLayer(UDataLayer* DataLayerToDelete);
+	void DeleteDataLayer(UDataLayer* DataLayerToDelete);
 
 	/**
 	 * Renames the provided DataLayer to the provided new label
@@ -571,7 +571,7 @@ public:
 	 * @param	NewDataLayerLabel	The new label for the DataLayer to be renamed
 	 */
 	UFUNCTION(BlueprintCallable, Category = DataLayers)
-	virtual bool RenameDataLayer(UDataLayer* DataLayer, const FName& NewDataLayerLabel);
+	bool RenameDataLayer(UDataLayer* DataLayer, const FName& NewDataLayerLabel);
 
 private:
 
diff --git a/Engine/Source/Editor/Experimental/EditorInteractiveToolsFramework/Public/EdModeInteractiveToolsContext.h b/Engine/Source/Editor/Experimental/EditorInteractiveToolsFramework/Public/EdModeInteractiveToolsContext.h
--- a/Engine/Source/Editor/Experimental/EditorInteractiveToolsFramework/Public/EdModeInteractiveToolsContext.h
+++ b/Engine/Source/Editor/Experimental/EditorInteractiveToolsFramework/Public/EdModeInteractiveToolsContext.h
@@ -34,66 +34,66 @@ public:
 	UEdModeInteractiveToolsContext();
 
 	UE_DEPRECATED(5.0, "You should no longer create your own EdModeInteractiveToolsContext; use the one in the FEditorModeTools::GetInteractiveToolsContext instead.")
-	virtual void InitializeContextFromEdMode(FEdMode* EditorModeIn,
+	void InitializeContextFromEdMode(FEdMode* EditorModeIn,
 		IToolsContextAssetAPI* UseAssetAPI = nullptr);
-	virtual void InitializeContextWithEditorModeManager(FEditorModeTools* InEditorModeManager,
+	void InitializeContextWithEditorModeManager(FEditorModeTools* InEditorModeManager,
 		IToolsContextAssetAPI* UseAssetAPI = nullptr);
-	virtual void ShutdownContext();
+	void ShutdownContext();
 
 	// default behavior is to accept active tool
-	virtual void TerminateActiveToolsOnPIEStart();
+	void TerminateActiveToolsOnPIEStart();
 
 	// default behavior is to accept active tool
-	virtual void TerminateActiveToolsOnSaveWorld();
+	void TerminateActiveToolsOnSaveWorld();
 
 	// default behavior is to accept active tool
-	virtual void TerminateActiveToolsOnWorldTearDown();
+	void TerminateActiveToolsOnWorldTearDown();
 
 	IToolsContextQueriesAPI* GetQueriesAPI() const { return QueriesAPI; }
 	IToolsContextTransactionsAPI* GetTransactionAPI() const { return TransactionAPI; }
 	IToolsContextAssetAPI* GetAssetAPI() const { return AssetAPI; }
 
-	virtual void PostInvalidation();
+	void PostInvalidation();
 
 	// UObject Interface
 	virtual UWorld* GetWorld() const override;
 
 	// call these from your UEdMode functions of the same name
 
-	virtual void Tick(FEditorViewportClient* ViewportClient, float DeltaTime);
-	virtual void Render(const FSceneView* View, FViewport* Viewport, FPrimitiveDrawInterface* PDI);
-	virtual void DrawHUD(FViewportClient* ViewportClient,FViewport* Viewport,const FSceneView* View, FCanvas* Canvas);
+	void Tick(FEditorViewportClient* ViewportClient, float DeltaTime);
+	void Render(const FSceneView* View, FViewport* Viewport, FPrimitiveDrawInterface* PDI);
+	void DrawHUD(FViewportClient* ViewportClient,FViewport* Viewport,const FSceneView* View, FCanvas* Canvas);
 
-	virtual bool ProcessEditDelete();
+	bool ProcessEditDelete();
 
-	virtual bool InputKey(FEditorViewportClient* ViewportClient, FViewport* Viewport, FKey Key, EInputEvent Event);
+	bool InputKey(FEditorViewportClient* ViewportClient, FViewport* Viewport, FKey Key, EInputEvent Event);
 
-	virtual bool MouseEnter(FEditorViewportClient* ViewportClient, FViewport* Viewport, int32 x, int32 y);
-	virtual bool MouseLeave(FEditorViewportClient* ViewportClient, FViewport* Viewport);
-	virtual bool MouseMove(FEditorViewportClient* ViewportClient, FViewport* Viewport, int32 x, int32 y);
+	bool MouseEnter(FEditorViewportClient* ViewportClient, FViewport* Viewport, int32 x, int32 y);
+	bool MouseLeave(FEditorViewportClient* ViewportClient, FViewport* Viewport);
+	bool MouseMove(FEditorViewportClient* ViewportClient, FViewport* Viewport, int32 x, int32 y);
 
-	virtual bool StartTracking(FEditorViewportClient* InViewportClient, FViewport* InViewport);
-	virtual bool CapturedMouseMove(FEditorViewportClient* InViewportClient, FViewport* InViewport, int32 InMouseX, int32 InMouseY);
-	virtual bool EndTracking(FEditorViewportClient* InViewportClient, FViewport* InViewport);
+	bool StartTracking(FEditorViewportClient* InViewportClient, FViewport* InViewport);
+	bool CapturedMouseMove(FEditorViewportClient* InViewportClient, FViewport* InViewport, int32 InMouseX, int32 InMouseY);
+	bool EndTracking(FEditorViewportClient* InViewportClient, FViewport* InViewport);
 
 
 	//
 	// Utility functions useful for hooking up to UICommand/etc
 	//
 
-	virtual bool CanStartTool(const FString ToolTypeIdentifier) const;
-	virtual bool HasActiveTool() const;
-	virtual FString GetActiveToolName() const;
-	virtual bool ActiveToolHasAccept() const;
-	virtual bool CanAcceptActiveTool() const;
-	virtual bool CanCancelActiveTool() const;
-	virtual bool CanCompleteActiveTool() const;
-	virtual void StartTool(const FString ToolTypeIdentifier);
-	virtual void EndTool(EToolShutdownType ShutdownType);
+	bool CanStartTool(const FString ToolTypeIdentifier) const;
+	bool HasActiveTool() const;
+	FString GetActiveToolName() const;
+	bool ActiveToolHasAccept() const;
+	bool CanAcceptActiveTool() const;
+	bool CanCancelActiveTool() const;
+	bool CanCompleteActiveTool() const;
+	void StartTool(const FString ToolTypeIdentifier);
+	void EndTool(EToolShutdownType ShutdownType);
 
-	virtual bool ShouldIgnoreHotkeys() const { return bInFlyMode; }
+	bool ShouldIgnoreHotkeys() const { return bInFlyMode; }
 
-	virtual FRay GetLastWorldRay() const;
+	FRay GetLastWorldRay() const;
 
 protected:
 	// we hide these 
diff --git a/Engine/Source/Editor/Sequencer/Private/Sequencer.h b/Engine/Source/Editor/Sequencer/Private/Sequencer.h
--- a/Engine/Source/Editor/Sequencer/Private/Sequencer.h
+++ b/Engine/Source/Editor/Sequencer/Private/Sequencer.h
@@ -361,7 +361,7 @@ public:
 	 * @param TargetObjectGuid	Object to be targeted on dropping
 	 * @return					If true, this asset should be consumed
 	 */
-	virtual bool OnHandleAssetDropped( UObject* DroppedAsset, const FGuid& TargetObjectGuid );
+	bool OnHandleAssetDropped( UObject* DroppedAsset, const FGuid& TargetObjectGuid );
 	
 	/**
 	 * Called to delete all moviescene data from a node
@@ -369,7 +369,7 @@ public:
 	 * @param NodeToBeDeleted	Node with data that should be deleted
 	 * @return true if anything was deleted, otherwise false.
 	 */
-	virtual bool OnRequestNodeDeleted( TSharedRef<const FSequencerDisplayNode> NodeToBeDeleted, const bool bKeepState );
+	bool OnRequestNodeDeleted( TSharedRef<const FSequencerDisplayNode> NodeToBeDeleted, const bool bKeepState );
 
 	/** Zooms to the edges of all currently selected sections and keys. */
 	void ZoomToFit();
diff --git a/Engine/Source/Runtime/Core/Public/HAL/MallocCallstackHandler.h b/Engine/Source/Runtime/Core/Public/HAL/MallocCallstackHandler.h
--- a/Engine/Source/Runtime/Core/Public/HAL/MallocCallstackHandler.h
+++ b/Engine/Source/Runtime/Core/Public/HAL/MallocCallstackHandler.h
@@ -207,7 +207,7 @@ protected:
 		return FPlatformTLS::GetTlsValue(DisabledTLS) != 0;
 	}
 
-	void TrackRealloc(void* OldPtr, void* NewPtr, uint32 NewSize, uint32 OldSize, int32 CallStackIndex);
+	virtual void TrackRealloc(void* OldPtr, void* NewPtr, uint32 NewSize, uint32 OldSize, int32 CallStackIndex);
 	virtual void TrackMalloc(void* Ptr, uint32 Size, int32 CallStackIndex) = 0;
 	virtual void TrackFree(void* Ptr, uint32 OldSize, int32 CallStackIndex) = 0;
 
diff --git a/Engine/Source/Runtime/D3D12RHI/Private/D3D12CommandContext.h b/Engine/Source/Runtime/D3D12RHI/Private/D3D12CommandContext.h
--- a/Engine/Source/Runtime/D3D12RHI/Private/D3D12CommandContext.h
+++ b/Engine/Source/Runtime/D3D12RHI/Private/D3D12CommandContext.h
@@ -683,10 +683,6 @@ public:
 	{
 		ContextRedirect(RHIUpdateTextureReference(TextureRef, NewTexture));
 	}
-	FORCEINLINE virtual void RHIClearMRTImpl(bool bClearColor, int32 NumClearColors, const FLinearColor* ColorArray, bool bClearDepth, float Depth, bool bClearStencil, uint32 Stencil)
-	{
-		ContextRedirect(RHIClearMRTImpl(bClearColor, NumClearColors, ColorArray, bClearDepth, Depth, bClearStencil, Stencil));
-	}
 
 	FORCEINLINE virtual void RHIWaitForTemporalEffect(const FName& InEffectName) final AFR_API_OVERRIDE
 	{
diff --git a/Engine/Source/Runtime/Engine/Classes/Engine/World.h b/Engine/Source/Runtime/Engine/Classes/Engine/World.h
--- a/Engine/Source/Runtime/Engine/Classes/Engine/World.h
+++ b/Engine/Source/Runtime/Engine/Classes/Engine/World.h
@@ -3510,7 +3510,7 @@ public:
 	class FAudioDevice* GetAudioDeviceRaw();
 
 	/** Return the URL of this level on the local machine. */
-	virtual FString GetLocalURL() const;
+	FString GetLocalURL() const;
 
 	/** Returns whether script is executing within the editor. */
 	bool IsPlayInEditor() const;
@@ -3538,7 +3538,7 @@ public:
 
 	// Return the URL of this level, which may possibly
 	// exist on a remote machine.
-	virtual FString GetAddressURL() const;
+	FString GetAddressURL() const;
 
 	/**
 	 * Called after GWorld has been set. Used to load, but not associate, all
@@ -3584,7 +3584,7 @@ public:
 	 * @param bAbsolute whether we are using relative or absolute travel
 	 * @param bShouldSkipGameNotify whether to notify the clients/game or not
 	 */
-	virtual bool ServerTravel(const FString& InURL, bool bAbsolute = false, bool bShouldSkipGameNotify = false);
+	bool ServerTravel(const FString& InURL, bool bAbsolute = false, bool bShouldSkipGameNotify = false);
 
 	/** seamlessly travels to the given URL by first loading the entry level in the background,
 	 * switching to it, and then loading the specified level. Does not disrupt network communication or disconnect clients.
@@ -3774,7 +3774,7 @@ public:
 	void GetMatineeActors( TArray<AMatineeActor*>& OutMatineeActors );
 
 	/** Updates all physics constraint actor joint locations.  */
-	virtual void UpdateConstraintActors();
+	void UpdateConstraintActors();
 
 	/** Gets all LightMaps and ShadowMaps associated with this world. Specify the level or leave null for persistent */
 	void GetLightMapsAndShadowMaps(ULevel* Level, TArray<UTexture2D*>& OutLightMapsAndShadowMaps, bool bForceLazyLoad = true);
diff --git a/Engine/Source/Runtime/Engine/Public/ParticleHelper.h b/Engine/Source/Runtime/Engine/Public/ParticleHelper.h
--- a/Engine/Source/Runtime/Engine/Public/ParticleHelper.h
+++ b/Engine/Source/Runtime/Engine/Public/ParticleHelper.h
@@ -2522,7 +2522,7 @@ public:
 	/**
 	 *	Called when the rendering thread removes the dynamic data from the scene.
 	 */
-	virtual void ReleaseRenderThreadResources();
+	void ReleaseRenderThreadResources();
 
 	void UpdateData(FParticleDynamicData* NewDynamicData);
 	void UpdateData_RenderThread(FParticleDynamicData* NewDynamicData);
diff --git a/Engine/Source/Runtime/Experimental/Chaos/Public/Chaos/BVHParticles.h b/Engine/Source/Runtime/Experimental/Chaos/Public/Chaos/BVHParticles.h
--- a/Engine/Source/Runtime/Experimental/Chaos/Public/Chaos/BVHParticles.h
+++ b/Engine/Source/Runtime/Experimental/Chaos/Public/Chaos/BVHParticles.h
@@ -51,11 +51,6 @@ namespace Chaos
 
 		CHAOS_API void Serialize(FChaosArchive& Ar);
 
-		CHAOS_API virtual void Serialize(FArchive& Ar)
-		{
-			check(false); //Aggregate simplicial require FChaosArchive - check false by default
-		}
-
 	private:
 		CHAOS_API TBVHParticles(const TBVHParticles<T, d>& Other);
 
diff --git a/Engine/Source/Runtime/Experimental/Chaos/Public/Chaos/ImplicitObjectTransformed.h b/Engine/Source/Runtime/Experimental/Chaos/Public/Chaos/ImplicitObjectTransformed.h
--- a/Engine/Source/Runtime/Experimental/Chaos/Public/Chaos/ImplicitObjectTransformed.h
+++ b/Engine/Source/Runtime/Experimental/Chaos/Public/Chaos/ImplicitObjectTransformed.h
@@ -182,20 +182,20 @@ public:
 		MTransform = InTransform;
 	}
 
-	virtual void AccumulateAllImplicitObjects(TArray<Pair<const FImplicitObject*, TRigidTransform<T, d>>>& Out, const TRigidTransform<T, d>& ParentTM) const
+	virtual void AccumulateAllImplicitObjects(TArray<Pair<const FImplicitObject*, TRigidTransform<T, d>>>& Out, const TRigidTransform<T, d>& ParentTM) const override
 	{
 		const TRigidTransform<T, d> NewTM = MTransform * ParentTM;
 		MObject->AccumulateAllImplicitObjects(Out, NewTM);
 	}
 
-	virtual void AccumulateAllSerializableImplicitObjects(TArray<Pair<TSerializablePtr<FImplicitObject>, TRigidTransform<T, d>>>& Out, const TRigidTransform<T, d>& ParentTM, TSerializablePtr<FImplicitObject> This) const
+	virtual void AccumulateAllSerializableImplicitObjects(TArray<Pair<TSerializablePtr<FImplicitObject>, TRigidTransform<T, d>>>& Out, const TRigidTransform<T, d>& ParentTM, TSerializablePtr<FImplicitObject> This) const override
 	{
 		check(bSerializable);
 		const TRigidTransform<T, d> NewTM = MTransform * ParentTM;
 		TImplicitObjectTransformAccumulateSerializableHelper(Out, MObject, NewTM);
 	}
 
-	virtual void FindAllIntersectingObjects(TArray < Pair<const FImplicitObject*, TRigidTransform<T, d>>>& Out, const TAABB<T, d>& LocalBounds) const
+	virtual void FindAllIntersectingObjects(TArray < Pair<const FImplicitObject*, TRigidTransform<T, d>>>& Out, const TAABB<T, d>& LocalBounds) const override
 	{
 		const TAABB<T, d> SubobjectBounds = LocalBounds.TransformedAABB(MTransform.Inverse());
 		int32 NumOut = Out.Num();
diff --git a/Engine/Source/Runtime/Foliage/Private/FoliageComponent.cpp b/Engine/Source/Runtime/Foliage/Private/FoliageComponent.cpp
--- a/Engine/Source/Runtime/Foliage/Private/FoliageComponent.cpp
+++ b/Engine/Source/Runtime/Foliage/Private/FoliageComponent.cpp
@@ -31,7 +31,7 @@ public:
 	{}
 
 	/** Accessor used by the rendering thread when setting foliage parameters for rendering. */
-	virtual void GetFoliageParameters(FVector& OutFoliageImpluseDirection, FVector4& OutFoliageNormalizedRotationAxisAndAngle) const
+	void GetFoliageParameters(FVector& OutFoliageImpluseDirection, FVector4& OutFoliageNormalizedRotationAxisAndAngle) const
 	{
 		OutFoliageImpluseDirection = FoliageImpluseDirection;
 		OutFoliageNormalizedRotationAxisAndAngle = FoliageNormalizedRotationAxisAndAngle;
diff --git a/Engine/Source/Runtime/OpenGLDrv/Public/OpenGLDrv.h b/Engine/Source/Runtime/OpenGLDrv/Public/OpenGLDrv.h
--- a/Engine/Source/Runtime/OpenGLDrv/Public/OpenGLDrv.h
+++ b/Engine/Source/Runtime/OpenGLDrv/Public/OpenGLDrv.h
@@ -503,9 +503,9 @@ public:
 	virtual void RHICopyBufferRegion(FRHIBuffer* DestBuffer, uint64 DstOffset, FRHIBuffer* SourceBuffer, uint64 SrcOffset, uint64 NumBytes) final override;
 
 	// FOpenGLDynamicRHI interface.
-	virtual FTexture2DRHIRef RHICreateTexture2DFromResource(EPixelFormat Format, uint32 SizeX, uint32 SizeY, uint32 NumMips, uint32 NumSamples, uint32 NumSamplesTileMem, const FClearValueBinding& ClearValueBinding, GLuint Resource, ETextureCreateFlags Flags);
-	virtual FTexture2DRHIRef RHICreateTexture2DArrayFromResource(EPixelFormat Format, uint32 SizeX, uint32 SizeY, uint32 ArraySize, uint32 NumMips, uint32 NumSamples, uint32 NumSamplesTileMem, const FClearValueBinding& ClearValueBinding, GLuint Resource, ETextureCreateFlags Flags);
-	virtual FTextureCubeRHIRef RHICreateTextureCubeFromResource(EPixelFormat Format, uint32 Size, bool bArray, uint32 ArraySize, uint32 NumMips, uint32 NumSamples, uint32 NumSamplesTileMem, const FClearValueBinding& ClearValueBinding, GLuint Resource, ETextureCreateFlags Flags);
+	FTexture2DRHIRef RHICreateTexture2DFromResource(EPixelFormat Format, uint32 SizeX, uint32 SizeY, uint32 NumMips, uint32 NumSamples, uint32 NumSamplesTileMem, const FClearValueBinding& ClearValueBinding, GLuint Resource, ETextureCreateFlags Flags);
+	FTexture2DRHIRef RHICreateTexture2DArrayFromResource(EPixelFormat Format, uint32 SizeX, uint32 SizeY, uint32 ArraySize, uint32 NumMips, uint32 NumSamples, uint32 NumSamplesTileMem, const FClearValueBinding& ClearValueBinding, GLuint Resource, ETextureCreateFlags Flags);
+	FTextureCubeRHIRef RHICreateTextureCubeFromResource(EPixelFormat Format, uint32 Size, bool bArray, uint32 ArraySize, uint32 NumMips, uint32 NumSamples, uint32 NumSamplesTileMem, const FClearValueBinding& ClearValueBinding, GLuint Resource, ETextureCreateFlags Flags);
 
 	// Inline copy
 	virtual void RHICopyToStagingBuffer(FRHIVertexBuffer* SourceBufferRHI, FRHIStagingBuffer* DestinationStagingBufferRHI, uint32 InOffset, uint32 InNumBytes) final override;
diff --git a/Engine/Source/Runtime/Windows/AudioMixerXAudio2/Private/AudioMixerPlatformWindows.cpp b/Engine/Source/Runtime/Windows/AudioMixerXAudio2/Private/AudioMixerPlatformWindows.cpp
--- a/Engine/Source/Runtime/Windows/AudioMixerXAudio2/Private/AudioMixerPlatformWindows.cpp
+++ b/Engine/Source/Runtime/Windows/AudioMixerXAudio2/Private/AudioMixerPlatformWindows.cpp
@@ -32,7 +32,7 @@ public:
 		}
 	}
 
-	virtual ~FWindowsMMNotificationClient()
+	~FWindowsMMNotificationClient()
 	{
 		if (DeviceEnumerator)
 		{
diff --git a/Engine/Source/Runtime/Windows/AudioMixerXAudio2/Private/AudioMixerPlatformXAudio2.h b/Engine/Source/Runtime/Windows/AudioMixerXAudio2/Private/AudioMixerPlatformXAudio2.h
--- a/Engine/Source/Runtime/Windows/AudioMixerXAudio2/Private/AudioMixerPlatformXAudio2.h
+++ b/Engine/Source/Runtime/Windows/AudioMixerXAudio2/Private/AudioMixerPlatformXAudio2.h
@@ -38,7 +38,7 @@ namespace Audio
 	{
 	public:
 		FXAudio2VoiceCallback() {}
-		virtual ~FXAudio2VoiceCallback() {}
+		~FXAudio2VoiceCallback() {}
 
 	private:
 		void STDCALL OnVoiceProcessingPassStart(UINT32 BytesRequired) {}
diff --git a/Engine/Source/Runtime/Windows/D3D11RHI/Private/D3D11Query.cpp b/Engine/Source/Runtime/Windows/D3D11RHI/Private/D3D11Query.cpp
--- a/Engine/Source/Runtime/Windows/D3D11RHI/Private/D3D11Query.cpp
+++ b/Engine/Source/Runtime/Windows/D3D11RHI/Private/D3D11Query.cpp
@@ -130,7 +130,7 @@ private:
 		: CurrentId(0)
 	{}
 
-	virtual ~FD3D11RenderQueryBatcher() = default;
+	~FD3D11RenderQueryBatcher() = default;
 };
 
 void D3D11RHIQueryBatcherPerFrameCleanup()
diff --git a/Engine/Source/Runtime/Windows/XAudio2/Private/XAudio2Support.h b/Engine/Source/Runtime/Windows/XAudio2/Private/XAudio2Support.h
--- a/Engine/Source/Runtime/Windows/XAudio2/Private/XAudio2Support.h
+++ b/Engine/Source/Runtime/Windows/XAudio2/Private/XAudio2Support.h
@@ -56,7 +56,7 @@ public:
 		}
 	}
 
-	virtual ~FMMNotificationClient()
+	~FMMNotificationClient()
 	{
 		if (DeviceEnumerator)
 		{
@@ -798,7 +798,7 @@ struct FXAudioDeviceProperties final : public IDeviceChangedListener
 #endif
 	}
 	
-	virtual ~FXAudioDeviceProperties()
+	~FXAudioDeviceProperties()
 	{
 #if PLATFORM_WINDOWS
 		if (NotificationClient)
