Fixed DataLayerOutliner not updating when copy/pasting an actor with DataLayers

diff --git a/Engine/Source/Editor/DataLayerEditor/Private/DataLayer/DataLayerEditorSubsystem.cpp b/Engine/Source/Editor/DataLayerEditor/Private/DataLayer/DataLayerEditorSubsystem.cpp
--- a/Engine/Source/Editor/DataLayerEditor/Private/DataLayer/DataLayerEditorSubsystem.cpp
+++ b/Engine/Source/Editor/DataLayerEditor/Private/DataLayer/DataLayerEditorSubsystem.cpp
@@ -105,7 +105,7 @@ void FDataLayersBroadcast::OnObjectPostEditChange(UObject* Object, FPropertyChan
 		}
 		else if (AActor* Actor = Cast<AActor>(Object))
 		{
-			bRefresh = Actor->IsPropertyChangedAffectingDataLayers(PropertyChangedEvent);
+			bRefresh = Actor->IsPropertyChangedAffectingDataLayers(PropertyChangedEvent) || Actor->HasDataLayers();
 		}
 		if (bRefresh)
 		{
