- Use same actor filter as World Outliner in Data Layer outliner.
- Data Layer property type customization now relies on ProperyHandle IsEditable to be enabled

diff --git a/Engine/Source/Editor/DataLayerEditor/Private/DataLayer/DataLayerMode.cpp b/Engine/Source/Editor/DataLayerEditor/Private/DataLayer/DataLayerMode.cpp
--- a/Engine/Source/Editor/DataLayerEditor/Private/DataLayer/DataLayerMode.cpp
+++ b/Engine/Source/Editor/DataLayerEditor/Private/DataLayer/DataLayerMode.cpp
@@ -13,6 +13,7 @@
 #include "ScopedTransaction.h"
 #include "DragAndDrop/ActorDragDropOp.h"
 #include "DragAndDrop/FolderDragDropOp.h"
+#include "ActorMode.h"
 #include "EditorActorFolders.h"
 #include "Algo/Transform.h"
 #include "ToolMenus.h"
@@ -32,6 +33,11 @@ FDataLayerMode::FDataLayerMode(const FDataLayerModeParams& Params)
 	, DataLayerBrowser(Params.DataLayerBrowser)
 	, SpecifiedWorldToDisplay(Params.SpecifiedWorldToDisplay)
 {
+	SceneOutliner->AddFilter(MakeShared<TSceneOutlinerPredicateFilter<FDataLayerActorTreeItem>>(FDataLayerActorTreeItem::FFilterPredicate::CreateLambda([this](const AActor* Actor, const UDataLayer* DataLayer)
+	{
+		return FActorMode::IsActorDisplayable(SceneOutliner, Actor);
+	}), FSceneOutlinerFilter::EDefaultBehaviour::Pass));
+
 	DataLayerEditorSubsystem = UDataLayerEditorSubsystem::Get();
 	Rebuild();
 	const_cast<FSharedSceneOutlinerData&>(SceneOutliner->GetSharedData()).CustomDelete = FCustomSceneOutlinerDeleteDelegate::CreateRaw(this, &FDataLayerMode::DeleteItems);
diff --git a/Engine/Source/Editor/DataLayerEditor/Private/DataLayer/DataLayerPropertyTypeCustomization.cpp b/Engine/Source/Editor/DataLayerEditor/Private/DataLayer/DataLayerPropertyTypeCustomization.cpp
--- a/Engine/Source/Editor/DataLayerEditor/Private/DataLayer/DataLayerPropertyTypeCustomization.cpp
+++ b/Engine/Source/Editor/DataLayerEditor/Private/DataLayer/DataLayerPropertyTypeCustomization.cpp
@@ -5,7 +5,6 @@
 #include "DataLayer/DataLayerDragDropOp.h"
 #include "DataLayer/DataLayerEditorSubsystem.h"
 #include "WorldPartition/DataLayer/DataLayer.h"
-#include "WorldPartition/WorldPartitionSubsystem.h"
 #include "Algo/Accumulate.h"
 #include "Modules/ModuleManager.h"
 #include "PropertyHandle.h"
@@ -36,7 +35,6 @@ void FDataLayerPropertyTypeCustomization::CustomizeHeader(TSharedRef<IPropertyHa
 	.MaxDesiredWidth(TOptional<float>())
 	[
 		SNew(SDropTarget)
-		.IsEnabled_Lambda([]() { return GWorld ? UWorld::HasSubsystem<UWorldPartitionSubsystem>(GWorld) : false; })
 		.OnDrop(this, &FDataLayerPropertyTypeCustomization::OnDrop)
 		.OnAllowDrop(this, &FDataLayerPropertyTypeCustomization::OnVerifyDrag)
 		.OnIsRecognized(this, &FDataLayerPropertyTypeCustomization::OnVerifyDrag)
@@ -89,6 +87,7 @@ void FDataLayerPropertyTypeCustomization::CustomizeHeader(TSharedRef<IPropertyHa
 			]
 		]
 	];
+	HeaderRow.IsEnabled(TAttribute<bool>(StructPropertyHandle, &IPropertyHandle::IsEditable));
 }
 
 UDataLayer* FDataLayerPropertyTypeCustomization::GetDataLayerFromPropertyHandle() const
diff --git a/Engine/Source/Editor/SceneOutliner/Private/ActorMode.cpp b/Engine/Source/Editor/SceneOutliner/Private/ActorMode.cpp
--- a/Engine/Source/Editor/SceneOutliner/Private/ActorMode.cpp
+++ b/Engine/Source/Editor/SceneOutliner/Private/ActorMode.cpp
@@ -268,6 +268,11 @@ void FActorMode::SynchronizeActorSelection()
 }
 
 bool FActorMode::IsActorDisplayable(const AActor* Actor) const
+{
+	return FActorMode::IsActorDisplayable(SceneOutliner, Actor);
+}
+
+bool FActorMode::IsActorDisplayable(const SSceneOutliner* SceneOutliner, const AActor* Actor)
 {
 	static const FName SequencerActorTag(TEXT("SequencerActor"));
 
diff --git a/Engine/Source/Editor/SceneOutliner/Public/ActorMode.h b/Engine/Source/Editor/SceneOutliner/Public/ActorMode.h
--- a/Engine/Source/Editor/SceneOutliner/Public/ActorMode.h
+++ b/Engine/Source/Editor/SceneOutliner/Public/ActorMode.h
@@ -56,6 +56,8 @@ public:
 	virtual void OnFilterTextChanged(const FText& InFilterText) override;
 
 	virtual int32 GetTypeSortPriority(const ISceneOutlinerTreeItem& Item) const override;
+
+	static bool IsActorDisplayable(const SSceneOutliner* SceneOutliner, const AActor* Actor);
 private:
 	/** Called when the user selects a world in the world picker menu */
 	void OnSelectWorld(TWeakObjectPtr<UWorld> World);
