- Added details in Data Layer outliner (needed to setup advanced Data Layer settings)
- Changed Data Layer to not be dynamically loaded by default (same as old layers)

diff --git a/Engine/Source/Editor/DataLayerEditor/Private/DataLayer/DataLayerMode.cpp b/Engine/Source/Editor/DataLayerEditor/Private/DataLayer/DataLayerMode.cpp
--- a/Engine/Source/Editor/DataLayerEditor/Private/DataLayer/DataLayerMode.cpp
+++ b/Engine/Source/Editor/DataLayerEditor/Private/DataLayer/DataLayerMode.cpp
@@ -63,6 +63,22 @@ int32 FDataLayerMode::GetTypeSortPriority(const ISceneOutlinerTreeItem& Item) co
 	return -1;
 }
 
+FText FDataLayerMode::GetStatusText() const
+{
+	if (SelectedDataLayersSet.Num() == 1)
+	{
+		if (const UDataLayer* DataLayer = SelectedDataLayersSet.Array()[0].Get())
+		{
+			return FText::FromName(DataLayer->GetDataLayerLabel());
+		}
+	}
+	else if (SelectedDataLayersSet.Num() > 1)
+	{
+		return FText::Format(LOCTEXT("DetailsMultipleDataLayersSelected", "{0} Data Layers Selected"), FText::AsNumber(SelectedDataLayersSet.Num()));
+	}
+	return LOCTEXT("DetailsNoDataLayerSelected", "No Data Layer Selected");
+}
+
 SDataLayerBrowser* FDataLayerMode::GetDataLayerBrowser() const
 {
 	return DataLayerBrowser;
@@ -149,8 +165,21 @@ void FDataLayerMode::DeleteItems(const TArray<TWeakPtr<ISceneOutlinerTreeItem>>&
 	}
 	else if (!DataLayersToDelete.IsEmpty())
 	{
-		const FScopedTransaction Transaction(LOCTEXT("DeleteDataLayers", "Delete DataLayers"));
-		DataLayerEditorSubsystem->DeleteDataLayers(DataLayersToDelete);
+		int32 PrevDeleteCount = SelectedDataLayersSet.Num();
+		for (UDataLayer* DataLayerToDelete : DataLayersToDelete)
+		{
+			SelectedDataLayersSet.Remove(DataLayerToDelete);
+		}
+		
+		{
+			const FScopedTransaction Transaction(LOCTEXT("DeleteDataLayers", "Delete DataLayers"));
+			DataLayerEditorSubsystem->DeleteDataLayers(DataLayersToDelete);
+		}
+
+		if ((SelectedDataLayersSet.Num() != PrevDeleteCount) && DataLayerBrowser)
+		{
+			DataLayerBrowser->OnSelectionChanged(SelectedDataLayersSet);
+		}
 	}
 }
 
@@ -615,6 +644,10 @@ void FDataLayerMode::OnItemSelectionChanged(FSceneOutlinerTreeItemPtr TreeItem,
 	SelectedDataLayerActors.Empty();
 	Selection.ForEachItem<FDataLayerTreeItem>([this](const FDataLayerTreeItem& Item) { SelectedDataLayersSet.Add(Item.GetDataLayer()); });
 	Selection.ForEachItem<FDataLayerActorTreeItem>([this](const FDataLayerActorTreeItem& Item) { SelectedDataLayerActors.Add(FSelectedDataLayerActor(Item.GetDataLayer(), Item.GetActor())); });
+	if (DataLayerBrowser)
+	{
+		DataLayerBrowser->OnSelectionChanged(SelectedDataLayersSet);
+	}
 }
 
 void FDataLayerMode::Rebuild()
diff --git a/Engine/Source/Editor/DataLayerEditor/Private/DataLayer/DataLayerMode.h b/Engine/Source/Editor/DataLayerEditor/Private/DataLayer/DataLayerMode.h
--- a/Engine/Source/Editor/DataLayerEditor/Private/DataLayer/DataLayerMode.h
+++ b/Engine/Source/Editor/DataLayerEditor/Private/DataLayer/DataLayerMode.h
@@ -44,6 +44,9 @@ public:
 	virtual bool IsInteractive() const override { return true; }
 	virtual bool CanRename() const override { return true; }
 	virtual bool CanRenameItem(const ISceneOutlinerTreeItem& Item) const override { return (Item.IsValid() && (Item.IsA<FDataLayerTreeItem>())); }
+	virtual bool ShowStatusBar() const override { return true; }
+	virtual FText GetStatusText() const override;
+	virtual FSlateColor GetStatusTextColor() const override { return FSlateColor::UseForeground(); }
 
 	virtual void OnItemDoubleClick(FSceneOutlinerTreeItemPtr Item) override;
 	virtual void OnItemAdded(FSceneOutlinerTreeItemPtr Item) override;
diff --git a/Engine/Source/Editor/DataLayerEditor/Private/DataLayer/SDataLayerBrowser.cpp b/Engine/Source/Editor/DataLayerEditor/Private/DataLayer/SDataLayerBrowser.cpp
--- a/Engine/Source/Editor/DataLayerEditor/Private/DataLayer/SDataLayerBrowser.cpp
+++ b/Engine/Source/Editor/DataLayerEditor/Private/DataLayer/SDataLayerBrowser.cpp
@@ -8,11 +8,13 @@
 #include "DataLayerTreeItem.h"
 #include "DataLayerMode.h"
 #include "WorldPartition/DataLayer/DataLayerEditorPerProjectUserSettings.h"
+#include "WorldPartition/DataLayer/DataLayer.h"
 #include "DataLayerOutlinerIsDynamicallyLoadedColumn.h"
 #include "DataLayerOutlinerDeleteButtonColumn.h"
 #include "DataLayer/DataLayerEditorSubsystem.h"
 #include "WorldPartition/DataLayer/DataLayer.h"
 #include "WorldPartition/WorldPartitionSubsystem.h"
+#include "Modules/ModuleManager.h"
 #include "ScopedTransaction.h"
 #include "Editor.h"
 
@@ -47,6 +49,15 @@ void SDataLayerBrowser::Construct(const FArguments& InArgs)
 		return (ToggleModeButton.IsValid() && (ToggleModeButton->IsHovered() || ToggleModeButton->IsPressed())) ? FEditorStyle::GetSlateColor(InvertedForegroundName) : FSlateColor::UseForeground();
 	};
 
+	FPropertyEditorModule& PropertyModule = FModuleManager::LoadModuleChecked<FPropertyEditorModule>("PropertyEditor");
+	FDetailsViewArgs Args;
+	Args.NameAreaSettings = FDetailsViewArgs::HideNameArea;
+	Args.bAllowSearch = false;
+	Args.bHideSelectionTip = true;
+	Args.bShowActorLabel = false;
+	DetailsWidget = PropertyModule.CreateDetailView(Args);
+	DetailsWidget->SetVisibility(EVisibility::Visible);
+
 	//////////////////////////////////////////////////////////////////////////
 	//	DataLayer Contents Header
 	SAssignNew(DataLayerContentsHeader, SBorder)
@@ -104,18 +115,34 @@ void SDataLayerBrowser::Construct(const FArguments& InArgs)
 	InitOptions.ColumnMap.Add(FDataLayerOutlinerIsDynamicallyLoadedColumn::GetID(), FSceneOutlinerColumnInfo(ESceneOutlinerColumnVisibility::Visible, 1, FCreateSceneOutlinerColumn::CreateLambda([](ISceneOutliner& InSceneOutliner) { return MakeShareable(new FDataLayerOutlinerIsDynamicallyLoadedColumn(InSceneOutliner)); })));
 	InitOptions.ColumnMap.Add(FSceneOutlinerBuiltInColumnTypes::Label(), FSceneOutlinerColumnInfo(ESceneOutlinerColumnVisibility::Visible, 2));
 	InitOptions.ColumnMap.Add(FDataLayerOutlinerDeleteButtonColumn::GetID(), FSceneOutlinerColumnInfo(ESceneOutlinerColumnVisibility::Visible, 20, FCreateSceneOutlinerColumn::CreateLambda([](ISceneOutliner& InSceneOutliner) { return MakeShareable(new FDataLayerOutlinerDeleteButtonColumn(InSceneOutliner)); })));
-
-	auto CreateDataLayerBrowser = [&](const FSceneOutlinerInitializationOptions& InInitOptions)
-	{
-		return SNew(SDataLayerOutliner, InInitOptions).IsEnabled(FSlateApplication::Get().GetNormalExecutionAttribute());
-	};
+	DataLayerOutliner = SNew(SDataLayerOutliner, InitOptions).IsEnabled(FSlateApplication::Get().GetNormalExecutionAttribute());
 
 	SAssignNew(DataLayerContentsSection, SBorder)
 	.Padding(5)
 	.BorderImage(FEditorStyle::GetBrush("NoBrush"))
 	.Content()
 	[
-		CreateDataLayerBrowser(InitOptions)
+		// Data Layer Outliner
+		SNew(SSplitter)
+		.Orientation(Orient_Vertical)
+		.Style(FEditorStyle::Get(), "FoliageEditMode.Splitter")
+		+ SSplitter::Slot()
+		[
+			SNew(SVerticalBox)
+			+ SVerticalBox::Slot()
+			[
+				DataLayerOutliner.ToSharedRef()
+			]
+		]
+		// Details
+		+SSplitter::Slot()
+		[
+			SNew(SVerticalBox)
+			+ SVerticalBox::Slot()
+			[
+				DetailsWidget.ToSharedRef()
+			]
+		]
 	];
 
 	//////////////////////////////////////////////////////////////////////////
@@ -129,6 +156,20 @@ void SDataLayerBrowser::Construct(const FArguments& InArgs)
 	SetupDataLayerMode(Mode);
 }
 
+void SDataLayerBrowser::OnSelectionChanged(TSet<TWeakObjectPtr<const UDataLayer>>& InSelectedDataLayersSet)
+{
+	SelectedDataLayersSet = InSelectedDataLayersSet;
+	TArray<UObject*> SelectedDataLayers;
+	for (const auto& WeakDataLayer : SelectedDataLayersSet)
+	{
+		if (WeakDataLayer.IsValid())
+		{
+			SelectedDataLayers.Add(const_cast<UDataLayer*>(WeakDataLayer.Get()));
+		}
+	}
+	DetailsWidget->SetObjects(SelectedDataLayers, /*bForceRefresh*/ true);
+}
+
 
 void SDataLayerBrowser::SetupDataLayerMode(EDataLayerBrowserMode InNewMode)
 {
diff --git a/Engine/Source/Editor/DataLayerEditor/Private/DataLayer/SDataLayerBrowser.h b/Engine/Source/Editor/DataLayerEditor/Private/DataLayer/SDataLayerBrowser.h
--- a/Engine/Source/Editor/DataLayerEditor/Private/DataLayer/SDataLayerBrowser.h
+++ b/Engine/Source/Editor/DataLayerEditor/Private/DataLayer/SDataLayerBrowser.h
@@ -13,6 +13,7 @@ class ISceneOutliner;
 class SButton;
 class SVerticalBox;
 class SBorder;
+class UDataLayer;
 
 enum class EDataLayerBrowserMode
 {
@@ -47,6 +48,8 @@ public:
 	DECLARE_EVENT_OneParam(SDataLayerBrowser, FOnModeChanged, EDataLayerBrowserMode /*Mode*/);
 	virtual FOnModeChanged& OnModeChanged() final { return ModeChanged; }
 
+	void OnSelectionChanged(TSet<TWeakObjectPtr<const UDataLayer>>& SelectedDataLayersSet);
+
 private:
 
 	void SetupDataLayerMode(EDataLayerBrowserMode InNewMode);
@@ -54,6 +57,9 @@ private:
 	/**	Fires whenever one or more actor DataLayer changes */
 	FOnModeChanged ModeChanged;
 
+	TSet<TWeakObjectPtr<const UDataLayer>> SelectedDataLayersSet;
+	TSharedPtr<class ISceneOutliner> DataLayerOutliner;
+	TSharedPtr<class IDetailsView> DetailsWidget;
 	TSharedPtr<SButton> ToggleModeButton;
 	TSharedPtr<SVerticalBox> ContentAreaBox;
 	TSharedPtr<SBorder> DataLayerContentsSection;
diff --git a/Engine/Source/Runtime/Engine/Private/WorldPartition/DataLayer/DataLayer.cpp b/Engine/Source/Runtime/Engine/Private/WorldPartition/DataLayer/DataLayer.cpp
--- a/Engine/Source/Runtime/Engine/Private/WorldPartition/DataLayer/DataLayer.cpp
+++ b/Engine/Source/Runtime/Engine/Private/WorldPartition/DataLayer/DataLayer.cpp
@@ -13,7 +13,7 @@ UDataLayer::UDataLayer(const FObjectInitializer& ObjectInitializer)
 #if WITH_EDITORONLY_DATA
 , DataLayerLabel(GetFName())
 , bIsVisible(true)
-, bIsDynamicallyLoaded(true)
+, bIsDynamicallyLoaded(false)
 , bIsDynamicallyLoadedInEditor(true)
 #endif
 {
diff --git a/Engine/Source/Runtime/Engine/Public/WorldPartition/DataLayer/DataLayer.h b/Engine/Source/Runtime/Engine/Public/WorldPartition/DataLayer/DataLayer.h
--- a/Engine/Source/Runtime/Engine/Public/WorldPartition/DataLayer/DataLayer.h
+++ b/Engine/Source/Runtime/Engine/Public/WorldPartition/DataLayer/DataLayer.h
@@ -35,11 +35,11 @@ private:
 	FName DataLayerLabel;
 
 	/** Whether actors associated with the DataLayer are visible in the viewport */
-	UPROPERTY()
+	UPROPERTY(Category = DataLayer, EditAnywhere)
 	uint32 bIsVisible : 1;
 
 	/** Whether the DataLayer affects actor runtime loading */
-	UPROPERTY()
+	UPROPERTY(Category = DataLayer, EditAnywhere)
 	uint32 bIsDynamicallyLoaded : 1;
 
 #if WITH_EDITORONLY_DATA
