When running in PIE show which runtime data layers are active in the DataLayer outliner.

diff --git a/Engine/Source/Editor/DataLayerEditor/Private/DataLayer/SDataLayerTreeLabel.cpp b/Engine/Source/Editor/DataLayerEditor/Private/DataLayer/SDataLayerTreeLabel.cpp
--- a/Engine/Source/Editor/DataLayerEditor/Private/DataLayer/SDataLayerTreeLabel.cpp
+++ b/Engine/Source/Editor/DataLayerEditor/Private/DataLayer/SDataLayerTreeLabel.cpp
@@ -4,6 +4,8 @@
 #include "DataLayer/DataLayerEditorSubsystem.h"
 #include "DataLayerTreeItem.h"
 #include "Widgets/Text/SInlineEditableTextBlock.h"
+#include "WorldPartition/DataLayer/DataLayerSubsystem.h"
+#include "Math/ColorList.h"
 #include "ISceneOutliner.h"
 #include "ISceneOutlinerMode.h"
 #include "ScopedTransaction.h"
@@ -87,7 +89,15 @@ void SDataLayerTreeLabel::Construct(const FArguments& InArgs, FDataLayerTreeItem
 FText SDataLayerTreeLabel::GetDisplayText() const
 {
 	const UDataLayer* DataLayer = DataLayerPtr.Get();
-	return DataLayer ? FText::FromName(DataLayer->GetDataLayerLabel()) : LOCTEXT("DataLayerLabelForMissingDataLayer", "(Deleted Data Layer)");
+	bool bIsDataLayerActive = false;
+	if (DataLayer && DataLayer->IsDynamicallyLoaded() && DataLayer->GetWorld()->IsPlayInEditor())
+	{
+		const UDataLayerSubsystem* DataLayerSubsystem = DataLayer->GetWorld()->GetSubsystem<UDataLayerSubsystem>();
+		bIsDataLayerActive = DataLayerSubsystem->IsDataLayerActive(DataLayer);
+	}
+	static const FText DataLayerActive = LOCTEXT("DataLayerActive", "(Active)");
+	static const FText DataLayerDeleted = LOCTEXT("DataLayerLabelForMissingDataLayer", "(Deleted Data Layer)");
+	return DataLayer ? FText::Format(LOCTEXT("DataLayerDisplayText", "{0} {1}"), FText::FromName(DataLayer->GetDataLayerLabel()), bIsDataLayerActive ? DataLayerActive : FText::GetEmpty()) : DataLayerDeleted;
 }
 
 FText SDataLayerTreeLabel::GetTooltipText() const
@@ -144,6 +154,14 @@ FSlateColor SDataLayerTreeLabel::GetForegroundColor() const
 	}
 
 	const UDataLayer* DataLayer = DataLayerPtr.Get();
+	if (DataLayer && DataLayer->IsDynamicallyLoaded() && DataLayer->GetWorld()->IsPlayInEditor())
+	{
+		const UDataLayerSubsystem* DataLayerSubsystem = DataLayer->GetWorld()->GetSubsystem<UDataLayerSubsystem>();
+		if (DataLayerSubsystem->IsDataLayerActive(DataLayer))
+		{
+			return FColorList::LimeGreen;
+		}
+	}
 	return (!DataLayer || !DataLayer->GetWorld()) ? FLinearColor(0.2f, 0.2f, 0.25f) : FSlateColor::UseForeground();
 }
 
