Data Layers
- Added the concept of locked layers, which can be set to prevent the user to add or remove actors from the list, and also prevent deleting or renaming the layer. Users can't create locked data layers, only meant to be created by custom code to group actors together. The first use case will be procedural generation, to easily gather previous generation actors from actor descriptors.
- Make sure we always have a world data layer actors when converting a map wo world partition or when loading a partiioned world.

diff --git a/Engine/Source/Editor/DataLayerEditor/Private/DataLayer/DataLayerMode.cpp b/Engine/Source/Editor/DataLayerEditor/Private/DataLayer/DataLayerMode.cpp
--- a/Engine/Source/Editor/DataLayerEditor/Private/DataLayer/DataLayerMode.cpp
+++ b/Engine/Source/Editor/DataLayerEditor/Private/DataLayer/DataLayerMode.cpp
@@ -63,6 +63,16 @@ int32 FDataLayerMode::GetTypeSortPriority(const ISceneOutlinerTreeItem& Item) co
 	return -1;
 }
 
+bool FDataLayerMode::CanRenameItem(const ISceneOutlinerTreeItem& Item) const 
+{
+	if (Item.IsValid() && (Item.IsA<FDataLayerTreeItem>()))
+	{
+		FDataLayerTreeItem* DataLayerTreeItem = (FDataLayerTreeItem*)&Item;
+		return !DataLayerTreeItem->GetDataLayer()->IsLocked();
+	}
+	return false;
+}
+
 FText FDataLayerMode::GetStatusText() const
 {
 	if (SelectedDataLayersSet.Num() == 1)
@@ -141,7 +151,7 @@ void FDataLayerMode::DeleteItems(const TArray<TWeakPtr<ISceneOutlinerTreeItem>>&
 		{
 			UDataLayer* DataLayer = DataLayerActorItem->GetDataLayer();
 			AActor* Actor = DataLayerActorItem->GetActor();
-			if (DataLayer && Actor)
+			if (DataLayer && !DataLayer->IsLocked() && Actor)
 			{
 				ActorsToRemoveFromDataLayer.FindOrAdd(DataLayer).Add(Actor);
 			}
@@ -150,7 +160,10 @@ void FDataLayerMode::DeleteItems(const TArray<TWeakPtr<ISceneOutlinerTreeItem>>&
 		{
 			if (UDataLayer* DataLayer = DataLayerItem->GetDataLayer())
 			{
-				DataLayersToDelete.Add(DataLayer);
+				if (!DataLayer->IsLocked())
+				{
+					DataLayersToDelete.Add(DataLayer);
+				}
 			}
 		}
 	}
@@ -236,10 +249,29 @@ FSceneOutlinerDragValidationInfo FDataLayerMode::ValidateDrop(const ISceneOutlin
 			{
 				return FSceneOutlinerDragValidationInfo(ESceneOutlinerDropCompatibility::IncompatibleGeneric, FText());
 			}
-			if (SceneOutliner->GetTree().IsItemSelected(const_cast<ISceneOutlinerTreeItem&>(DropTarget).AsShared()) && (GetSelectedDataLayers(SceneOutliner).Num() > 1))
+			
+			if (DataLayerTarget->IsLocked())
 			{
-				return FSceneOutlinerDragValidationInfo(ESceneOutlinerDropCompatibility::Compatible, LOCTEXT("AssignToDataLayers", "Assign to Selected Data Layers"));
+				return FSceneOutlinerDragValidationInfo(ESceneOutlinerDropCompatibility::IncompatibleGeneric, LOCTEXT("CantReassignSystemDataLayer", "Can't reassign actors from system Data Layer"));
 			}
+			
+			if (GetSelectedDataLayers(SceneOutliner).Num() > 1)
+			{
+				TArray<UDataLayer*> SelectedDataLayers = GetSelectedDataLayers(SceneOutliner);
+				for (UDataLayer* SelectedDataLayer : SelectedDataLayers)
+				{
+					if (SelectedDataLayer->IsLocked())
+					{
+						return FSceneOutlinerDragValidationInfo(ESceneOutlinerDropCompatibility::IncompatibleGeneric, LOCTEXT("CantReassignSystemDataLayer", "Can't reassign actors from system Data Layer"));
+					}
+				}
+
+				if (SceneOutliner->GetTree().IsItemSelected(const_cast<ISceneOutlinerTreeItem&>(DropTarget).AsShared()))
+				{
+					return FSceneOutlinerDragValidationInfo(ESceneOutlinerDropCompatibility::Compatible, LOCTEXT("AssignToDataLayers", "Assign to Selected Data Layers"));
+				}
+			}
+
 			return FSceneOutlinerDragValidationInfo(ESceneOutlinerDropCompatibility::Compatible, FText::Format(LOCTEXT("AssignToDataLayer", "Assign to Data Layer \"{0}\""), FText::FromName(DataLayerTarget->GetDataLayerLabel())));
 		}
 		else
@@ -533,7 +565,16 @@ void FDataLayerMode::RegisterContextMenu()
 								const FScopedTransaction Transaction(LOCTEXT("DeleteSelectedDataLayers", "Delete Selected DataLayers"));
 								DataLayerEditorSubsystem->DeleteDataLayers(SelectedDataLayers);
 							}}),
-						FCanExecuteAction::CreateLambda([=] { return !SelectedDataLayers.IsEmpty(); })
+						FCanExecuteAction::CreateLambda([=] { 
+							for (UDataLayer* SelectedDataLayer : SelectedDataLayers)
+							{
+								if (SelectedDataLayer->IsLocked())
+								{
+									return false;
+								}
+							}
+							return !SelectedDataLayers.IsEmpty();
+						})
 					));
 
 				Section.AddMenuEntry("RenameSelectedDataLayer", LOCTEXT("RenameSelectedDataLayer", "Rename Selected DataLayer"), FText(), FSlateIcon(),
@@ -548,7 +589,7 @@ void FDataLayerMode::RegisterContextMenu()
 									SceneOutliner->ScrollItemIntoView(ItemToRename);
 								}
 							}}),
-						FCanExecuteAction::CreateLambda([=] { return SelectedDataLayers.Num() == 1; })
+						FCanExecuteAction::CreateLambda([=] { return (SelectedDataLayers.Num() == 1) && !SelectedDataLayers[0]->IsLocked(); })
 					));
 
 				Section.AddSeparator("SectionsSeparator");
diff --git a/Engine/Source/Editor/DataLayerEditor/Private/DataLayer/DataLayerMode.h b/Engine/Source/Editor/DataLayerEditor/Private/DataLayer/DataLayerMode.h
--- a/Engine/Source/Editor/DataLayerEditor/Private/DataLayer/DataLayerMode.h
+++ b/Engine/Source/Editor/DataLayerEditor/Private/DataLayer/DataLayerMode.h
@@ -43,7 +43,7 @@ public:
 	virtual bool SupportsKeyboardFocus() const override { return true; }
 	virtual bool IsInteractive() const override { return true; }
 	virtual bool CanRename() const override { return true; }
-	virtual bool CanRenameItem(const ISceneOutlinerTreeItem& Item) const override { return (Item.IsValid() && (Item.IsA<FDataLayerTreeItem>())); }
+	virtual bool CanRenameItem(const ISceneOutlinerTreeItem& Item) const override;
 	virtual bool ShowStatusBar() const override { return true; }
 	virtual FText GetStatusText() const override;
 	virtual FSlateColor GetStatusTextColor() const override { return FSlateColor::UseForeground(); }
diff --git a/Engine/Source/Editor/DataLayerEditor/Private/DataLayer/DataLayerPropertyTypeCustomizationHelper.cpp b/Engine/Source/Editor/DataLayerEditor/Private/DataLayer/DataLayerPropertyTypeCustomizationHelper.cpp
--- a/Engine/Source/Editor/DataLayerEditor/Private/DataLayer/DataLayerPropertyTypeCustomizationHelper.cpp
+++ b/Engine/Source/Editor/DataLayerEditor/Private/DataLayer/DataLayerPropertyTypeCustomizationHelper.cpp
@@ -35,12 +35,15 @@ TSharedRef<SWidget> FDataLayerPropertyTypeCustomizationHelper::CreateDataLayerMe
 		for (const TWeakObjectPtr<UDataLayer>& WeakDataLayer : AllDataLayers)
 		{
 			const UDataLayer* DataLayerPtr = WeakDataLayer.Get();
-			MenuBuilder.AddMenuEntry(
-				UDataLayer::GetDataLayerText(DataLayerPtr),
-				FText(),
-				FSlateIcon(FEditorStyle::GetStyleSetName(), "DataLayer.Icon16x"),
-				FUIAction(FExecuteAction::CreateLambda([DataLayerPtr, OnDataLayerSelectedFunction]() { OnDataLayerSelectedFunction(DataLayerPtr); }))
-			);
+			if (!DataLayerPtr || !DataLayerPtr->IsLocked())
+			{
+				MenuBuilder.AddMenuEntry(
+					UDataLayer::GetDataLayerText(DataLayerPtr),
+					FText(),
+					FSlateIcon(FEditorStyle::GetStyleSetName(), "DataLayer.Icon16x"),
+					FUIAction(FExecuteAction::CreateLambda([DataLayerPtr, OnDataLayerSelectedFunction]() { OnDataLayerSelectedFunction(DataLayerPtr); }))
+				);
+				}
 		}
 	}
 	MenuBuilder.EndSection();
diff --git a/Engine/Source/Editor/DataLayerEditor/Private/DataLayer/SDataLayerBrowser.cpp b/Engine/Source/Editor/DataLayerEditor/Private/DataLayer/SDataLayerBrowser.cpp
--- a/Engine/Source/Editor/DataLayerEditor/Private/DataLayer/SDataLayerBrowser.cpp
+++ b/Engine/Source/Editor/DataLayerEditor/Private/DataLayer/SDataLayerBrowser.cpp
@@ -164,7 +164,11 @@ void SDataLayerBrowser::OnSelectionChanged(TSet<TWeakObjectPtr<const UDataLayer>
 	{
 		if (WeakDataLayer.IsValid())
 		{
-			SelectedDataLayers.Add(const_cast<UDataLayer*>(WeakDataLayer.Get()));
+			UDataLayer* DataLayer = const_cast<UDataLayer*>(WeakDataLayer.Get());
+			if (!DataLayer->IsLocked())
+			{
+				SelectedDataLayers.Add(DataLayer);
+			}
 		}
 	}
 	DetailsWidget->SetObjects(SelectedDataLayers, /*bForceRefresh*/ true);
diff --git a/Engine/Source/Editor/DataLayerEditor/Private/DataLayer/SDataLayerTreeLabel.cpp b/Engine/Source/Editor/DataLayerEditor/Private/DataLayer/SDataLayerTreeLabel.cpp
--- a/Engine/Source/Editor/DataLayerEditor/Private/DataLayer/SDataLayerTreeLabel.cpp
+++ b/Engine/Source/Editor/DataLayerEditor/Private/DataLayer/SDataLayerTreeLabel.cpp
@@ -154,12 +154,19 @@ FSlateColor SDataLayerTreeLabel::GetForegroundColor() const
 	}
 
 	const UDataLayer* DataLayer = DataLayerPtr.Get();
-	if (DataLayer && DataLayer->IsDynamicallyLoaded() && DataLayer->GetWorld()->IsPlayInEditor())
+	if (DataLayer)
 	{
-		const UDataLayerSubsystem* DataLayerSubsystem = DataLayer->GetWorld()->GetSubsystem<UDataLayerSubsystem>();
-		if (DataLayerSubsystem->IsDataLayerActive(DataLayer))
+		if (DataLayer->IsLocked())
+		{
+			return FColorList::DimGrey;
+		}
+		else if (DataLayer->IsDynamicallyLoaded() && DataLayer->GetWorld()->IsPlayInEditor())
 		{
-			return FColorList::LimeGreen;
+			const UDataLayerSubsystem* DataLayerSubsystem = DataLayer->GetWorld()->GetSubsystem<UDataLayerSubsystem>();
+			if (DataLayerSubsystem->IsDataLayerActive(DataLayer))
+			{
+				return FColorList::LimeGreen;
+			}
 		}
 	}
 	return (!DataLayer || !DataLayer->GetWorld()) ? FLinearColor(0.2f, 0.2f, 0.25f) : FSlateColor::UseForeground();
diff --git a/Engine/Source/Runtime/Engine/Classes/Engine/World.h b/Engine/Source/Runtime/Engine/Classes/Engine/World.h
--- a/Engine/Source/Runtime/Engine/Classes/Engine/World.h
+++ b/Engine/Source/Runtime/Engine/Classes/Engine/World.h
@@ -1448,6 +1448,9 @@ private:
 
 	/** Utility function to cleanup streaming levels that point to invalid level packages */
 	void RepairStreamingLevels();
+
+	/** Utility function that is used to ensure that a World has the correct WorldDataLayers */
+	void RepairWorldDataLayers();
 	
 #if INCLUDE_CHAOS
 	/** Utility function that is used to ensure that a World has the correct ChaosActor */
diff --git a/Engine/Source/Runtime/Engine/Private/World.cpp b/Engine/Source/Runtime/Engine/Private/World.cpp
--- a/Engine/Source/Runtime/Engine/Private/World.cpp
+++ b/Engine/Source/Runtime/Engine/Private/World.cpp
@@ -106,6 +106,7 @@
 	#include "PIEPreviewDeviceProfileSelectorModule.h"
 	#include "Editor/FixupLazyObjectPtrForPIEArchive.h"
 	#include "AssetCompilingManager.h"
+	#include "WorldPartition/DataLayer/WorldDataLayers.h"
 #endif
 
 
@@ -1097,6 +1098,7 @@ void UWorld::PostLoad()
 #if WITH_EDITOR
 	RepairWorldSettings();
 	RepairStreamingLevels();
+	RepairWorldDataLayers();
 #endif
 #if INCLUDE_CHAOS
 	//RepairChaosActors();
@@ -1514,6 +1516,14 @@ void UWorld::RepairWorldSettings()
 	check(GetWorldSettings());
 }
 
+void UWorld::RepairWorldDataLayers()
+{
+	if (const UWorldPartition* WorldPartition = GetWorldPartition())
+	{
+		AWorldDataLayers::Get(this, /*bCreateIfNotFound*/true);
+	}
+}
+
 void UWorld::InitWorld(const InitializationValues IVS)
 {
 	if (!ensure(!bIsWorldInitialized))
@@ -1598,6 +1608,7 @@ void UWorld::InitWorld(const InitializationValues IVS)
 #if WITH_EDITOR
 	RepairWorldSettings();
 	RepairStreamingLevels();
+	RepairWorldDataLayers();
 #endif
 #if INCLUDE_CHAOS
 	//RepairChaosActors();
diff --git a/Engine/Source/Runtime/Engine/Private/WorldPartition/DataLayer/WorldDataLayers.cpp b/Engine/Source/Runtime/Engine/Private/WorldPartition/DataLayer/WorldDataLayers.cpp
--- a/Engine/Source/Runtime/Engine/Private/WorldPartition/DataLayer/WorldDataLayers.cpp
+++ b/Engine/Source/Runtime/Engine/Private/WorldPartition/DataLayer/WorldDataLayers.cpp
@@ -117,16 +117,15 @@ TArray<const UDataLayer*> AWorldDataLayers::GetDataLayerObjects(const TArray<FAc
 	return OutDataLayers;
 }
 
-UDataLayer* AWorldDataLayers::CreateDataLayer()
+UDataLayer* AWorldDataLayers::CreateDataLayer(FName InName, EObjectFlags InObjectFlags)
 {
 	Modify();
 
 	// Make sure new DataLayer name (not label) is unique and never re-used so that actors still referencing on deleted DataLayer's don't get valid again.
 	const FName DataLayerUniqueName = *FString::Format(TEXT("DataLayer_{0}"), { FGuid::NewGuid().ToString()});
-	UDataLayer* NewDataLayer = NewObject<UDataLayer>(this, DataLayerUniqueName, RF_Transactional);
+	UDataLayer* NewDataLayer = NewObject<UDataLayer>(this, DataLayerUniqueName, RF_Transactional | InObjectFlags);
 	check(NewDataLayer != NULL);
-	const FName DefaultLabel = TEXT("DataLayer");
-	FName DataLayerLabel = GenerateUniqueDataLayerLabel(DefaultLabel);
+	FName DataLayerLabel = GenerateUniqueDataLayerLabel(InName);
 	NewDataLayer->SetDataLayerLabel(DataLayerLabel);
 	NewDataLayer->SetVisible(true);
 	WorldDataLayers.Add(NewDataLayer);
diff --git a/Engine/Source/Runtime/Engine/Public/WorldPartition/DataLayer/DataLayer.h b/Engine/Source/Runtime/Engine/Public/WorldPartition/DataLayer/DataLayer.h
--- a/Engine/Source/Runtime/Engine/Public/WorldPartition/DataLayer/DataLayer.h
+++ b/Engine/Source/Runtime/Engine/Public/WorldPartition/DataLayer/DataLayer.h
@@ -21,6 +21,7 @@ public:
 	void SetVisible(bool bInIsVisible);
 	void SetIsDynamicallyLoaded(bool bInIsDynamicallyLoaded);
 	void SetIsDynamicallyLoadedInEditor(bool bInIsDynamicallyLoadedInEditor);
+	void SetIsLocked(bool bInIsLocked) { bIsLocked = bInIsLocked; }
 
 	bool IsDynamicallyLoadedInEditor() const { return !IsDynamicallyLoaded() || bIsDynamicallyLoadedInEditor; }
 	bool ShouldGenerateHLODs() const { return IsDynamicallyLoaded() && bGeneratesHLODs; }
@@ -33,6 +34,7 @@ public:
 	bool IsVisible() const { return bIsVisible; }
 	bool IsDynamicallyLoaded() const { return bIsDynamicallyLoaded; }
 	bool IsInitiallyActive() const { return IsDynamicallyLoaded() && bIsInitiallyActive; }
+	bool IsLocked() const { return bIsLocked; }
 
 private:
 	/** The display name of the DataLayer */
@@ -56,6 +58,10 @@ private:
 	UPROPERTY(Transient)
 	uint32 bIsDynamicallyLoadedInEditor : 1;
 
+	/** Whether this data layer is locked, which means the user can't change actors assignation, remove or rename it */
+	UPROPERTY()
+	uint32 bIsLocked : 1;
+
 	/** Whether HLODs should be generated for actors in this layer */
 	UPROPERTY(Category = DataLayer, EditAnywhere, meta=(EditConditionHides, EditCondition="bIsDynamicallyLoaded"))
 	uint32 bGeneratesHLODs : 1;
diff --git a/Engine/Source/Runtime/Engine/Public/WorldPartition/DataLayer/WorldDataLayers.h b/Engine/Source/Runtime/Engine/Public/WorldPartition/DataLayer/WorldDataLayers.h
--- a/Engine/Source/Runtime/Engine/Public/WorldPartition/DataLayer/WorldDataLayers.h
+++ b/Engine/Source/Runtime/Engine/Public/WorldPartition/DataLayer/WorldDataLayers.h
@@ -25,7 +25,7 @@ public:
 
 #if WITH_EDITOR
 	static AWorldDataLayers* Get(UWorld* World, bool bCreateIfNotFound);
-	UDataLayer* CreateDataLayer();
+	UDataLayer* CreateDataLayer(FName InName = TEXT("DataLayer"), EObjectFlags InObjectFlags = RF_NoFlags);
 	bool RemoveDataLayer(UDataLayer* InDataLayer);
 	bool RemoveDataLayers(const TArray<UDataLayer*>& InDataLayers);
 	FName GenerateUniqueDataLayerLabel(const FName& InDataLayerLabel) const;
