Data Layers Loading
- Disabled AInfo grid placement, runtime grid and data layers properties as they can't be changed by the user.
- Don't allow assigning data layers to AInfo.

diff --git a/Engine/Source/Editor/DataLayerEditor/Private/DataLayer/DataLayerMode.cpp b/Engine/Source/Editor/DataLayerEditor/Private/DataLayer/DataLayerMode.cpp
--- a/Engine/Source/Editor/DataLayerEditor/Private/DataLayer/DataLayerMode.cpp
+++ b/Engine/Source/Editor/DataLayerEditor/Private/DataLayer/DataLayerMode.cpp
@@ -242,6 +242,14 @@ FSceneOutlinerDragValidationInfo FDataLayerMode::ValidateDrop(const ISceneOutlin
 	TArray<AActor*> PayloadActors = GetActorsFromOperation(Payload.SourceOperation, true);
 	if (!PayloadActors.IsEmpty())
 	{
+		for (AActor* Actor : PayloadActors)
+		{
+			if (!Actor->SupportsDataLayer())
+			{
+				return FSceneOutlinerDragValidationInfo(ESceneOutlinerDropCompatibility::IncompatibleGeneric, LOCTEXT("ActorCantBeAssignedToDataLayer", "Can't assign actors to Data Layer"));
+			}
+		}
+
 		if (const FDataLayerTreeItem* DataLayerItem = DropTarget.CastTo<FDataLayerTreeItem>())
 		{
 			const UDataLayer* DataLayerTarget = DataLayerItem->GetDataLayer();
@@ -252,7 +260,7 @@ FSceneOutlinerDragValidationInfo FDataLayerMode::ValidateDrop(const ISceneOutlin
 			
 			if (DataLayerTarget->IsLocked())
 			{
-				return FSceneOutlinerDragValidationInfo(ESceneOutlinerDropCompatibility::IncompatibleGeneric, LOCTEXT("CantReassignSystemDataLayer", "Can't reassign actors from system Data Layer"));
+				return FSceneOutlinerDragValidationInfo(ESceneOutlinerDropCompatibility::IncompatibleGeneric, LOCTEXT("CantReassignLockedDataLayer", "Can't reassign actors from locked Data Layer"));
 			}
 			
 			if (GetSelectedDataLayers(SceneOutliner).Num() > 1)
@@ -262,7 +270,7 @@ FSceneOutlinerDragValidationInfo FDataLayerMode::ValidateDrop(const ISceneOutlin
 				{
 					if (SelectedDataLayer->IsLocked())
 					{
-						return FSceneOutlinerDragValidationInfo(ESceneOutlinerDropCompatibility::IncompatibleGeneric, LOCTEXT("CantReassignSystemDataLayer", "Can't reassign actors from system Data Layer"));
+						return FSceneOutlinerDragValidationInfo(ESceneOutlinerDropCompatibility::IncompatibleGeneric, LOCTEXT("CantReassignLockedDataLayer", "Can't reassign actors from locked Data Layer"));
 					}
 				}
 
diff --git a/Engine/Source/Runtime/Engine/Classes/GameFramework/Actor.h b/Engine/Source/Runtime/Engine/Classes/GameFramework/Actor.h
--- a/Engine/Source/Runtime/Engine/Classes/GameFramework/Actor.h
+++ b/Engine/Source/Runtime/Engine/Classes/GameFramework/Actor.h
@@ -781,7 +781,7 @@ protected:
 	FGuid ActorGuid;
 
 	/** DataLayers the actor belongs to.*/
-	UPROPERTY(EditAnywhere, AdvancedDisplay, Category = Actor)
+	UPROPERTY(EditAnywhere, AdvancedDisplay, Category = DataLayers)
 	TArray<FActorDataLayer> DataLayers;
 
 public:
@@ -1055,6 +1055,7 @@ public:
 	bool AddDataLayer(const UDataLayer* DataLayer);
 	bool RemoveDataLayer(const UDataLayer* DataLayer);
 	bool ContainsDataLayer(const UDataLayer* DataLayer) const;
+	virtual bool SupportsDataLayer() const { return true; }
 	bool HasDataLayers() const;
 	bool HasValidDataLayers() const;
 	bool HasAllDataLayers(const TArray<const UDataLayer*>& DataLayers) const;
diff --git a/Engine/Source/Runtime/Engine/Classes/GameFramework/Info.h b/Engine/Source/Runtime/Engine/Classes/GameFramework/Info.h
--- a/Engine/Source/Runtime/Engine/Classes/GameFramework/Info.h
+++ b/Engine/Source/Runtime/Engine/Classes/GameFramework/Info.h
@@ -16,7 +16,7 @@
  * Info is the base class of an Actor that isn't meant to have a physical representation in the world, used primarily
  * for "manager" type classes that hold settings data about the world, but might need to be an Actor for replication purposes.
  */
-UCLASS(abstract, hidecategories=(Input, Movement, Collision, Rendering, LOD, "Utilities|Transformation"), showcategories=("Input|MouseInput", "Input|TouchInput"), MinimalAPI, NotBlueprintable)
+UCLASS(abstract, hidecategories=(Input, Movement, Collision, Rendering, LOD, WorldPartition, DataLayers, "Utilities|Transformation"), showcategories=("Input|MouseInput", "Input|TouchInput"), MinimalAPI, NotBlueprintable)
 class AInfo : public AActor
 {
 	GENERATED_UCLASS_BODY()
@@ -34,6 +34,7 @@ public:
 
 #if WITH_EDITOR
 	virtual EActorGridPlacement GetDefaultGridPlacement() const override { return EActorGridPlacement::AlwaysLoaded; }
+	virtual bool SupportsDataLayer() const override { return false; }
 #endif
 
 public:
diff --git a/Engine/Source/Runtime/Engine/Classes/GameFramework/WorldSettings.h b/Engine/Source/Runtime/Engine/Classes/GameFramework/WorldSettings.h
--- a/Engine/Source/Runtime/Engine/Classes/GameFramework/WorldSettings.h
+++ b/Engine/Source/Runtime/Engine/Classes/GameFramework/WorldSettings.h
@@ -477,11 +477,11 @@ public:
 	 * Enables tools for creating a partitioned world. 
 	 * Level has to be saved and all sub-levels removed before enabling this option.
 	 */
-	UPROPERTY(EditAnywhere, BlueprintReadOnly, Category = WorldPartition)
+	UPROPERTY(EditAnywhere, BlueprintReadOnly, Category=WorldPartitionSetup)
 	uint8 bEnableWorldPartition:1;
 
 	/** Class of WorldPartitionStreamingPolicy to be used to manage world partition streaming. */
-	UPROPERTY(EditAnywhere, Category = WorldPartition, AdvancedDisplay, meta = (editcondition = "bEnableWorldPartition"))
+	UPROPERTY(EditAnywhere, Category=WorldPartitionSetup, AdvancedDisplay, meta=(editcondition = "bEnableWorldPartition"))
 	TSubclassOf<UWorldPartitionStreamingPolicy> WorldPartitionStreamingPolicyClass;
 
 	/** Additional transform applied when applying LevelStreaming Transform to LevelInstance */
@@ -545,7 +545,7 @@ protected:
 	UPROPERTY(Transient)
 	TObjectPtr<UNavigationSystemConfig> NavigationSystemConfigOverride;
 
-	UPROPERTY(EditAnywhere, Category = WorldPartition, NoClear, meta = (NoResetToDefault), Instanced)
+	UPROPERTY(EditAnywhere, Category=WorldPartitionSetup, NoClear, meta=(NoResetToDefault), Instanced)
 	TObjectPtr<UWorldPartition> WorldPartition;
 
 public:
diff --git a/Engine/Source/Runtime/Engine/Private/ActorEditor.cpp b/Engine/Source/Runtime/Engine/Private/ActorEditor.cpp
--- a/Engine/Source/Runtime/Engine/Private/ActorEditor.cpp
+++ b/Engine/Source/Runtime/Engine/Private/ActorEditor.cpp
@@ -71,11 +71,16 @@ bool AActor::CanEditChange(const FProperty* PropertyThatWillChange) const
 				return false;
 			}
 		}
+	}
 
-		if (bIsGridPlacement && (GetDefaultGridPlacement() != EActorGridPlacement::None))
-		{
-			return false;
-		}
+	if (bIsGridPlacement && (GetDefaultGridPlacement() != EActorGridPlacement::None))
+	{
+		return false;
+	}
+
+	if (bIsDataLayers && !SupportsDataLayer())
+	{
+		return false;
 	}
 
 	return Super::CanEditChange(PropertyThatWillChange);
@@ -1163,7 +1168,7 @@ EDataValidationResult AActor::IsDataValid(TArray<FText>& ValidationErrors)
 bool AActor::AddDataLayer(const UDataLayer* DataLayer)
 {
 	bool bActorWasModified = false;
-	if (DataLayer && !ContainsDataLayer(DataLayer))
+	if (SupportsDataLayer() && DataLayer && !ContainsDataLayer(DataLayer))
 	{
 		if (!bActorWasModified)
 		{
@@ -1260,6 +1265,11 @@ bool AActor::HasAnyOfDataLayers(const TArray<FName>& DataLayerNames) const
 
 void AActor::FixupDataLayers()
 {
+	if (!SupportsDataLayer())
+	{
+		DataLayers.Empty();
+	}
+
 	if (const AWorldDataLayers* WorldDataLayers = AWorldDataLayers::Get(GetWorld()))
 	{
 		TSet<FName> ExistingDataLayers;
