Properly support dragging multiple actors to the data layer browser.

diff --git a/Engine/Source/Editor/DataLayerEditor/Private/DataLayer/DataLayerMode.cpp b/Engine/Source/Editor/DataLayerEditor/Private/DataLayer/DataLayerMode.cpp
--- a/Engine/Source/Editor/DataLayerEditor/Private/DataLayer/DataLayerMode.cpp
+++ b/Engine/Source/Editor/DataLayerEditor/Private/DataLayer/DataLayerMode.cpp
@@ -239,7 +239,7 @@ bool FDataLayerMode::ParseDragDrop(FSceneOutlinerDragDropPayload& OutPayload, co
 
 FSceneOutlinerDragValidationInfo FDataLayerMode::ValidateDrop(const ISceneOutlinerTreeItem& DropTarget, const FSceneOutlinerDragDropPayload& Payload) const
 {
-	TArray<AActor*> PayloadActors = GetActorsFromOperation(Payload.SourceOperation, true);
+	TArray<AActor*> PayloadActors = GetActorsFromOperation(Payload.SourceOperation);
 	if (!PayloadActors.IsEmpty())
 	{
 		for (AActor* Actor : PayloadActors)
