- Changed Data Layer Visibility to be transient
- Added Data Layer bIsInitiallyVisible persistent flag to control Data Layer visibility flag initial value when loading map
- Fixed bug when changing map, Data Layer visibility was not properly reflected on actors
- Better Data Layer properties categorisation (Runtime vs Editor)
- Changing Data Layer IsDynamicallyLoaded flag is now only possible in its Details view
- Allow IsDynamicallyLoadedInEditor for any type of Data Layer

diff --git a/Engine/Source/Editor/DataLayerEditor/Private/DataLayer/DataLayerEditorSubsystem.cpp b/Engine/Source/Editor/DataLayerEditor/Private/DataLayer/DataLayerEditorSubsystem.cpp
--- a/Engine/Source/Editor/DataLayerEditor/Private/DataLayer/DataLayerEditorSubsystem.cpp
+++ b/Engine/Source/Editor/DataLayerEditor/Private/DataLayer/DataLayerEditorSubsystem.cpp
@@ -206,6 +206,7 @@ void UDataLayerEditorSubsystem::EditorMapChange()
 		World->PersistentLevel->OnLoadedActorAddedToLevelEvent.AddLambda([this](AActor& InActor) { InitializeNewActorDataLayers(&InActor); });
 	}
 	DataLayerChanged.Broadcast(EDataLayerAction::Reset, NULL, NAME_None);
+	UpdateAllActorsVisibility(true, true);
 }
 
 void UDataLayerEditorSubsystem::EditorRefreshDataLayerBrowser()
@@ -920,7 +921,7 @@ void UDataLayerEditorSubsystem::SetDataLayersVisibility(const TArray<UDataLayer*
 
 		if (DataLayer->IsVisible() != bIsVisible)
 		{
-			DataLayer->Modify();
+			DataLayer->Modify(/*bAlswaysMarkDirty*/false);
 			DataLayer->SetVisible(bIsVisible);
 			DataLayerChanged.Broadcast(EDataLayerAction::Modify, DataLayer, "bIsVisible");
 			bChangeOccurred = true;
diff --git a/Engine/Source/Editor/DataLayerEditor/Private/DataLayer/DataLayerOutlinerIsDynamicallyLoadedColumn.cpp b/Engine/Source/Editor/DataLayerEditor/Private/DataLayer/DataLayerOutlinerIsDynamicallyLoadedColumn.cpp
--- a/Engine/Source/Editor/DataLayerEditor/Private/DataLayer/DataLayerOutlinerIsDynamicallyLoadedColumn.cpp
+++ b/Engine/Source/Editor/DataLayerEditor/Private/DataLayer/DataLayerOutlinerIsDynamicallyLoadedColumn.cpp
@@ -43,69 +43,21 @@ const TSharedRef<SWidget> FDataLayerOutlinerIsDynamicallyLoadedColumn::Construct
 			.AutoWidth()
 			.VAlign(VAlign_Center)
 			[
-				SNew(SButton)
-				.ContentPadding(0)
-				.ButtonStyle(FEditorStyle::Get(), "NoBorder")
-				.OnClicked_Lambda([this, TreeItem]()
+				SNew(SImage)
+				.ToolTipText(LOCTEXT("IsDynamicallyLoadedButtonToolTip", "Whether the Data Layer affects actor runtime loading"))
+				.Image_Lambda([this, TreeItem]()
 				{
 					FDataLayerTreeItem* DataLayerTreeItem = TreeItem->CastTo<FDataLayerTreeItem>();
-					if (UDataLayer* DataLayer = DataLayerTreeItem->GetDataLayer())
+					const UDataLayer* DataLayer = DataLayerTreeItem->GetDataLayer();
+					if (DataLayer && DataLayer->IsDynamicallyLoaded())
 					{
-						bool bSuccess = false;
-						UDataLayerEditorSubsystem* DataLayerEditorSubsystem = UDataLayerEditorSubsystem::Get();
-						const auto& Tree = WeakSceneOutliner.Pin()->GetTree();
-						if (Tree.IsItemSelected(TreeItem))
-						{
-							// Toggle IsDynamicallyLoaded flag of selected DataLayers to the same state as the given DataLayer
-							bool bIsDynamicallyLoaded = DataLayer->IsDynamicallyLoaded();
-
-							TArray<UDataLayer*> AllSelectedDataLayers;
-							for (auto& SelectedItem : Tree.GetSelectedItems())
-							{
-								FDataLayerTreeItem* SelectedDataLayerTreeItem = SelectedItem->CastTo<FDataLayerTreeItem>();
-								UDataLayer* SelectedDataLayer = SelectedDataLayerTreeItem ? SelectedDataLayerTreeItem->GetDataLayer() : nullptr;
-								if (SelectedDataLayer && SelectedDataLayer->IsDynamicallyLoaded() == bIsDynamicallyLoaded)
-								{
-									AllSelectedDataLayers.Add(SelectedDataLayer);
-								}
-							}
-
-							const FScopedTransaction Transaction(LOCTEXT("ToggleDataLayersIsDynamicallyLoaded", "Toggle DataLayers Runtime Dynamically Loaded Flag"));
-							bSuccess = DataLayerEditorSubsystem->ToggleDataLayersIsDynamicallyLoaded(AllSelectedDataLayers);
-						}
-						else
-						{
-							const FScopedTransaction Transaction(LOCTEXT("ToggleDataLayerIsDynamicallyLoaded", "Toggle DataLayer Runtime Dynamically Loaded Flag"));
-							bSuccess = DataLayerEditorSubsystem->ToggleDataLayerIsDynamicallyLoaded(DataLayer);
-						}
-						if (!bSuccess)
-						{
-							// Cancelled : Undo last transaction
-							GEditor->Trans->Undo();
-						}
+						return FEditorStyle::GetBrush("DataLayer.DynamicallyLoaded");
 					}
-					return FReply::Handled();
-				})
-				.ToolTipText(LOCTEXT("IsDynamicallyLoadedButtonToolTip", "Toggle DataLayer Runtime Dynamically Loaded Flag"))
-				.HAlign(HAlign_Center)
-				.VAlign(VAlign_Center)
-				.Content()
-				[
-					SNew(SImage)
-					.Image_Lambda([this, TreeItem]()
+					else
 					{
-						FDataLayerTreeItem* DataLayerTreeItem = TreeItem->CastTo<FDataLayerTreeItem>();
-						const UDataLayer* DataLayer = DataLayerTreeItem->GetDataLayer();
-						if (DataLayer && DataLayer->IsDynamicallyLoaded())
-						{
-							return FEditorStyle::GetBrush("DataLayer.DynamicallyLoaded");
-						}
-						else
-						{
-							return FEditorStyle::GetBrush("DataLayer.NotDynamicallyLoaded");
-						}
-					})
-				]
+						return FEditorStyle::GetBrush("DataLayer.NotDynamicallyLoaded");
+					}
+				})
 			]
 			+ SHorizontalBox::Slot()
 			.Padding(2, 0, 0, 0)
@@ -113,12 +65,6 @@ const TSharedRef<SWidget> FDataLayerOutlinerIsDynamicallyLoadedColumn::Construct
 			.VAlign(VAlign_Center)
 			[
 				SNew(SCheckBox)
-				.IsEnabled_Lambda([this, TreeItem]()
-				{
-					FDataLayerTreeItem* DataLayerTreeItem = TreeItem->CastTo<FDataLayerTreeItem>();
-					UDataLayer* DataLayer = DataLayerTreeItem->GetDataLayer();
-					return DataLayer && DataLayer->IsDynamicallyLoaded();
-				})
 				.IsChecked_Lambda([this, TreeItem]()
 				{
 					FDataLayerTreeItem* DataLayerTreeItem = TreeItem->CastTo<FDataLayerTreeItem>();
@@ -150,12 +96,12 @@ const TSharedRef<SWidget> FDataLayerOutlinerIsDynamicallyLoadedColumn::Construct
 								}
 							}
 
-							const FScopedTransaction Transaction(LOCTEXT("ToggleDataLayersIsDynamicallyLoadedInEditor", "Toggle DataLayers Editor Dynamically Loaded Flag"));
+							const FScopedTransaction Transaction(LOCTEXT("ToggleDataLayersIsDynamicallyLoadedInEditor", "Toggle Data Layers Dynamically Loaded In Editor Flag"));
 							bSuccess = DataLayerEditorSubsystem->ToggleDataLayersIsDynamicallyLoadedInEditor(AllSelectedDataLayers);
 						}
 						else
 						{
-							const FScopedTransaction Transaction(LOCTEXT("ToggleDataLayerIsDynamicallyLoadedInEditor", "Toggle DataLayer Editor Dynamically Loaded Flag"));
+							const FScopedTransaction Transaction(LOCTEXT("ToggleDataLayerIsDynamicallyLoadedInEditor", "Toggle Data Layer Dynamically Loaded In Editor Flag"));
 							bSuccess = DataLayerEditorSubsystem->ToggleDataLayerIsDynamicallyLoadedInEditor(DataLayer);
 						}
 						if (!bSuccess)
@@ -165,7 +111,7 @@ const TSharedRef<SWidget> FDataLayerOutlinerIsDynamicallyLoadedColumn::Construct
 						}
 					}
 				})
-				.ToolTipText(LOCTEXT("IsDynamicallyLoadedInEditorCheckBoxToolTip", "Toggle DataLayer Editor Dynamically Loaded Flag"))
+				.ToolTipText(LOCTEXT("IsDynamicallyLoadedInEditorCheckBoxToolTip", "Toggle Data Layer Dynamically Loaded In Editor Flag"))
 				.HAlign(HAlign_Center)
 			];
 	}
diff --git a/Engine/Source/Editor/DataLayerEditor/Private/DataLayer/SDataLayerTreeLabel.cpp b/Engine/Source/Editor/DataLayerEditor/Private/DataLayer/SDataLayerTreeLabel.cpp
--- a/Engine/Source/Editor/DataLayerEditor/Private/DataLayer/SDataLayerTreeLabel.cpp
+++ b/Engine/Source/Editor/DataLayerEditor/Private/DataLayer/SDataLayerTreeLabel.cpp
@@ -90,15 +90,15 @@ FText SDataLayerTreeLabel::GetDisplayText() const
 {
 	const UDataLayer* DataLayer = DataLayerPtr.Get();
 	bool bIsDataLayerActive = false;
-	EDataLayerState DataLayerState = EDataLayerState::Unloaded;
+	FText DataLayerStateText = FText::GetEmpty();
 	if (DataLayer && DataLayer->IsDynamicallyLoaded() && DataLayer->GetWorld()->IsPlayInEditor())
 	{
 		const UDataLayerSubsystem* DataLayerSubsystem = DataLayer->GetWorld()->GetSubsystem<UDataLayerSubsystem>();
-		DataLayerState = DataLayerSubsystem->GetDataLayerState(DataLayer);
+		DataLayerStateText = FText::Format(LOCTEXT("DataLayerActive", " ({0})"), StaticEnum<EDataLayerState>()->GetDisplayNameTextByValue((int64)DataLayerSubsystem->GetDataLayerState(DataLayer)));
 	}
-	const FText DataLayerStateText = FText::Format(LOCTEXT("DataLayerActive", "({0})"), StaticEnum<EDataLayerState>()->GetDisplayNameTextByValue((int64)DataLayerState));
+	
 	static const FText DataLayerDeleted = LOCTEXT("DataLayerLabelForMissingDataLayer", "(Deleted Data Layer)");
-	return DataLayer ? FText::Format(LOCTEXT("DataLayerDisplayText", "{0} {1}"), FText::FromName(DataLayer->GetDataLayerLabel()), DataLayerStateText) : DataLayerDeleted;
+	return DataLayer ? FText::Format(LOCTEXT("DataLayerDisplayText", "{0}{1}"), FText::FromName(DataLayer->GetDataLayerLabel()), DataLayerStateText) : DataLayerDeleted;
 }
 
 FText SDataLayerTreeLabel::GetTooltipText() const
@@ -182,14 +182,14 @@ bool SDataLayerTreeLabel::OnVerifyItemLabelChanged(const FText& InLabel, FText&
 {
 	if (InLabel.IsEmptyOrWhitespace())
 	{
-		OutErrorMessage = LOCTEXT("EmptyDataLayerLabel", "DataLayer must be given a name");
+		OutErrorMessage = LOCTEXT("EmptyDataLayerLabel", "Data Layer must be given a name");
 		return false;
 	}
 
 	UDataLayer* FoundDataLayer;
 	if (UDataLayerEditorSubsystem::Get()->TryGetDataLayerFromLabel(*InLabel.ToString(), FoundDataLayer) && FoundDataLayer != DataLayerPtr.Get())
 	{
-		OutErrorMessage = LOCTEXT("RenameFailed_AlreadyExists", "This DataLayer already exists");
+		OutErrorMessage = LOCTEXT("RenameFailed_AlreadyExists", "This Data Layer already exists");
 		return false;
 	}
 
diff --git a/Engine/Source/Runtime/Engine/Private/WorldPartition/DataLayer/DataLayer.cpp b/Engine/Source/Runtime/Engine/Private/WorldPartition/DataLayer/DataLayer.cpp
--- a/Engine/Source/Runtime/Engine/Private/WorldPartition/DataLayer/DataLayer.cpp
+++ b/Engine/Source/Runtime/Engine/Private/WorldPartition/DataLayer/DataLayer.cpp
@@ -11,15 +11,17 @@
 UDataLayer::UDataLayer(const FObjectInitializer& ObjectInitializer)
 : Super(ObjectInitializer)
 #if WITH_EDITORONLY_DATA
-, DataLayerLabel(GetFName())
-, InitialState(EDataLayerState::Unloaded)
-, bIsVisible(true)
-, bIsDynamicallyLoaded(false)
 , bIsInitiallyActive_DEPRECATED(false)
+, bIsVisible(true)
+, bIsInitiallyVisible(true)
 , bIsDynamicallyLoadedInEditor(true)
+, bIsLocked(false)
 , bGeneratesHLODs(false)
 , DefaultHLODLayer()
 #endif
+, DataLayerLabel(GetFName())
+, InitialState(EDataLayerState::Unloaded)
+, bIsDynamicallyLoaded(false)
 {
 }
 
@@ -32,6 +34,27 @@ void UDataLayer::PostLoad()
 	{
 		InitialState = EDataLayerState::Activated;
 	}
+
+	// Initialize bIsVisible with persistent flag bIsInitiallyVisible
+	bIsVisible = bIsInitiallyVisible;
+#endif
+}
+
+bool UDataLayer::IsInitiallyVisible() const
+{
+#if WITH_EDITOR
+	return bIsInitiallyVisible;
+#else
+	return false;
+#endif
+}
+
+bool UDataLayer::IsVisible() const
+{
+#if WITH_EDITOR
+	return bIsVisible;
+#else
+	return false;
 #endif
 }
 
@@ -51,11 +74,20 @@ void UDataLayer::SetVisible(bool bInIsVisible)
 {
 	if (bIsVisible != bInIsVisible)
 	{
-		Modify();
+		Modify(/*bAlwaysMarkDirty*/false);
 		bIsVisible = bInIsVisible;
 	}
 }
 
+void UDataLayer::SetIsInitiallyVisible(bool bInIsInitiallyVisible)
+{
+	if (bIsInitiallyVisible != bInIsInitiallyVisible)
+	{
+		Modify();
+		bIsInitiallyVisible = bInIsInitiallyVisible;
+	}
+}
+
 void UDataLayer::SetIsDynamicallyLoaded(bool bInIsDynamicallyLoaded)
 {
 	if (bIsDynamicallyLoaded != bInIsDynamicallyLoaded)
@@ -67,7 +99,7 @@ void UDataLayer::SetIsDynamicallyLoaded(bool bInIsDynamicallyLoaded)
 
 void UDataLayer::SetIsDynamicallyLoadedInEditor(bool bInIsDynamicallyLoadedInEditor)
 {
-	if (bIsDynamicallyLoadedInEditor != bInIsDynamicallyLoadedInEditor && (!bInIsDynamicallyLoadedInEditor || bIsDynamicallyLoaded))
+	if (bIsDynamicallyLoadedInEditor != bInIsDynamicallyLoadedInEditor)
 	{
 		Modify(false);
 		bIsDynamicallyLoadedInEditor = bInIsDynamicallyLoadedInEditor;
diff --git a/Engine/Source/Runtime/Engine/Public/WorldPartition/DataLayer/DataLayer.h b/Engine/Source/Runtime/Engine/Public/WorldPartition/DataLayer/DataLayer.h
--- a/Engine/Source/Runtime/Engine/Public/WorldPartition/DataLayer/DataLayer.h
+++ b/Engine/Source/Runtime/Engine/Public/WorldPartition/DataLayer/DataLayer.h
@@ -28,11 +28,12 @@ public:
 #if WITH_EDITOR
 	void SetDataLayerLabel(FName InDataLayerLabel);
 	void SetVisible(bool bInIsVisible);
+	void SetIsInitiallyVisible(bool bInIsInitiallyVisible);
 	void SetIsDynamicallyLoaded(bool bInIsDynamicallyLoaded);
 	void SetIsDynamicallyLoadedInEditor(bool bInIsDynamicallyLoadedInEditor);
 	void SetIsLocked(bool bInIsLocked) { bIsLocked = bInIsLocked; }
 
-	bool IsDynamicallyLoadedInEditor() const { return !IsDynamicallyLoaded() || bIsDynamicallyLoadedInEditor; }
+	bool IsDynamicallyLoadedInEditor() const { return bIsDynamicallyLoadedInEditor; }
 	bool ShouldGenerateHLODs() const { return IsDynamicallyLoaded() && bGeneratesHLODs; }
 	class UHLODLayer* GetDefaultHLODLayer() const { return ShouldGenerateHLODs() ? DefaultHLODLayer : nullptr; }
 
@@ -42,46 +43,41 @@ public:
 #endif
 	virtual void PostLoad() override;
 
-	UFUNCTION(Category = DataLayer, BlueprintCallable)
+	UFUNCTION(Category = "Data Layer", BlueprintCallable)
 	bool Equals(const FActorDataLayer& ActorDataLayer) const { return ActorDataLayer.Name == GetFName(); }
 	
-	UFUNCTION(Category = DataLayer, BlueprintCallable)
+	UFUNCTION(Category = "Data Layer", BlueprintCallable)
 	FName GetDataLayerLabel() const  { return DataLayerLabel; }
 
-	UFUNCTION(Category = DataLayer, BlueprintCallable)
-	bool IsVisible() const { return bIsVisible; }
+	UFUNCTION(Category = "Data Layer - Editor", BlueprintCallable)
+	bool IsInitiallyVisible() const;
 
-	UFUNCTION(Category = DataLayer, BlueprintCallable)
+	UFUNCTION(Category = "Data Layer - Editor", BlueprintCallable)
+	bool IsVisible() const;
+
+	UFUNCTION(Category = "Data Layer - Runtime", BlueprintCallable)
 	bool IsDynamicallyLoaded() const { return bIsDynamicallyLoaded; }
 
-	UFUNCTION(Category = DataLayer, BlueprintCallable, meta = (DeprecatedFunction, DeprecationMessage = "Use GetInitialState instead"))
+	UFUNCTION(Category = "Data Layer - Runtime", BlueprintCallable, meta = (DeprecatedFunction, DeprecationMessage = "Use GetInitialState instead"))
 	bool IsInitiallyActive() const { return IsDynamicallyLoaded() && GetInitialState() == EDataLayerState::Activated; }	
 
-	UFUNCTION(Category = DataLayer, BlueprintCallable)
+	UFUNCTION(Category = "Data Layer - Runtime", BlueprintCallable)
 	EDataLayerState GetInitialState() const { return IsDynamicallyLoaded() ? InitialState : EDataLayerState::Unloaded; }
 
 private:
-	/** The display name of the DataLayer */
+#if WITH_EDITORONLY_DATA
 	UPROPERTY()
-	FName DataLayerLabel;
-
-	UPROPERTY(Category = DataLayer, EditAnywhere, meta = (EditConditionHides, EditCondition = "bIsDynamicallyLoaded"))
-	EDataLayerState InitialState;
+	uint32 bIsInitiallyActive_DEPRECATED : 1;
 
 	/** Whether actors associated with the DataLayer are visible in the viewport */
-	UPROPERTY(Category = DataLayer, EditAnywhere)
+	UPROPERTY(Transient)
 	uint32 bIsVisible : 1;
 
-	/** Whether the DataLayer affects actor runtime loading */
-	UPROPERTY(Category = DataLayer, EditAnywhere)
-	uint32 bIsDynamicallyLoaded : 1;
+	/** Whether actors associated with the Data Layer should be initially visible in the viewport when loading the map */
+	UPROPERTY(Category = "Data Layer - Editor", EditAnywhere)
+	uint32 bIsInitiallyVisible : 1;
 
-#if WITH_EDITORONLY_DATA
-	/** Whether a dynamically loaded DataLayer should be initially active at runtime */
-	UPROPERTY()
-	uint32 bIsInitiallyActive_DEPRECATED : 1;
-	
-	/** Whether the DataLayer affects actor editor loading */
+	/** Whether the Data Layer affects actor editor loading */
 	UPROPERTY(Transient)
 	uint32 bIsDynamicallyLoadedInEditor : 1;
 
@@ -90,11 +86,22 @@ private:
 	uint32 bIsLocked : 1;
 
 	/** Whether HLODs should be generated for actors in this layer */
-	UPROPERTY(Category = DataLayer, EditAnywhere, meta=(EditConditionHides, EditCondition="bIsDynamicallyLoaded"))
+	UPROPERTY(Category = "Data Layer - HLODs", EditAnywhere, meta = (EditConditionHides, EditCondition = "bIsDynamicallyLoaded"))
 	uint32 bGeneratesHLODs : 1;
 
 	// Default HLOD layer
-	UPROPERTY(Category = DataLayer, EditAnywhere, meta=(EditConditionHides, EditCondition="bIsDynamicallyLoaded && bGeneratesHLODs"))
+	UPROPERTY(Category = "Data Layer - HLODs", EditAnywhere, meta = (EditConditionHides, EditCondition = "bIsDynamicallyLoaded && bGeneratesHLODs"))
 	TObjectPtr<class UHLODLayer> DefaultHLODLayer;
 #endif
+
+	/** The display name of the Data Layer */
+	UPROPERTY()
+	FName DataLayerLabel;
+
+	UPROPERTY(Category = "Data Layer - Runtime", EditAnywhere, meta = (EditConditionHides, EditCondition = "bIsDynamicallyLoaded"))
+	EDataLayerState InitialState;
+
+	/** Whether the Data Layer affects actor runtime loading */
+	UPROPERTY(Category = "Data Layer - Runtime", EditAnywhere)
+	uint32 bIsDynamicallyLoaded : 1;
 };
\ No newline at end of file
