Fixed Drag & Drop of Data Layers in Actor Data Layers array

diff --git a/Engine/Source/Editor/DataLayerEditor/Private/DataLayer/DataLayerPropertyTypeCustomization.cpp b/Engine/Source/Editor/DataLayerEditor/Private/DataLayer/DataLayerPropertyTypeCustomization.cpp
--- a/Engine/Source/Editor/DataLayerEditor/Private/DataLayer/DataLayerPropertyTypeCustomization.cpp
+++ b/Engine/Source/Editor/DataLayerEditor/Private/DataLayer/DataLayerPropertyTypeCustomization.cpp
@@ -5,6 +5,7 @@
 #include "DataLayer/DataLayerDragDropOp.h"
 #include "DataLayer/DataLayerEditorSubsystem.h"
 #include "WorldPartition/DataLayer/DataLayer.h"
+#include "DragAndDrop/CompositeDragDropOp.h"
 #include "Algo/Accumulate.h"
 #include "Modules/ModuleManager.h"
 #include "PropertyHandle.h"
@@ -138,18 +139,15 @@ void FDataLayerPropertyTypeCustomization::AssignDataLayer(const UDataLayer* InDa
 
 FReply FDataLayerPropertyTypeCustomization::OnDrop(TSharedPtr<FDragDropOperation> InDragDrop)
 {
-	if (InDragDrop.IsValid() && InDragDrop->IsOfType<FDataLayerDragDropOp>())
+	TSharedPtr<const FDataLayerDragDropOp> DataLayerDragDropOp = GetDataLayerDragDropOp(InDragDrop);
+	if (DataLayerDragDropOp.IsValid())
 	{
-		TSharedPtr<const FDataLayerDragDropOp> DataLayerDragDropOp = StaticCastSharedPtr<const FDataLayerDragDropOp>(InDragDrop);
-		if (DataLayerDragDropOp.IsValid())
+		const TArray<FName>& DataLayerLabels = DataLayerDragDropOp->DataLayerLabels;
+		if (ensure(DataLayerLabels.Num() == 1))
 		{
-			const TArray<FName>& DataLayerLabels = DataLayerDragDropOp->DataLayerLabels;
-			if (ensure(DataLayerLabels.Num() == 1))
+			if (const UDataLayer* DataLayerPtr = UDataLayerEditorSubsystem::Get()->GetDataLayerFromLabel(DataLayerLabels[0]))
 			{
-				if (const UDataLayer* DataLayerPtr = UDataLayerEditorSubsystem::Get()->GetDataLayerFromLabel(DataLayerLabels[0]))
-				{
-					AssignDataLayer(DataLayerPtr);
-				}
+				AssignDataLayer(DataLayerPtr);
 			}
 		}
 	}
@@ -158,12 +156,26 @@ FReply FDataLayerPropertyTypeCustomization::OnDrop(TSharedPtr<FDragDropOperation
 
 bool FDataLayerPropertyTypeCustomization::OnVerifyDrag(TSharedPtr<FDragDropOperation> InDragDrop)
 {
-	if (InDragDrop.IsValid() && InDragDrop->IsOfType<FDataLayerDragDropOp>())
+	TSharedPtr<const FDataLayerDragDropOp> DataLayerDragDropOp = GetDataLayerDragDropOp(InDragDrop);
+	return DataLayerDragDropOp.IsValid() && DataLayerDragDropOp->DataLayerLabels.Num() == 1;
+}
+
+TSharedPtr<const FDataLayerDragDropOp> FDataLayerPropertyTypeCustomization::GetDataLayerDragDropOp(TSharedPtr<FDragDropOperation> InDragDrop)
+{
+	TSharedPtr<const FDataLayerDragDropOp> DataLayerDragDropOp;
+	if (InDragDrop.IsValid())
 	{
-		TSharedPtr<const FDataLayerDragDropOp> DataLayerDragDropOp = StaticCastSharedPtr<const FDataLayerDragDropOp>(InDragDrop);
-		return DataLayerDragDropOp.IsValid() && DataLayerDragDropOp->DataLayerLabels.Num() == 1;
+		if (InDragDrop->IsOfType<FCompositeDragDropOp>())
+		{
+			TSharedPtr<const FCompositeDragDropOp> CompositeDragDropOp = StaticCastSharedPtr<const FCompositeDragDropOp>(InDragDrop);
+			DataLayerDragDropOp = CompositeDragDropOp->GetSubOp<FDataLayerDragDropOp>();
+		}
+		else if (InDragDrop->IsOfType<FDataLayerDragDropOp>())
+		{
+			DataLayerDragDropOp = StaticCastSharedPtr<const FDataLayerDragDropOp>(InDragDrop);
+		}
 	}
-	return false;
+	return DataLayerDragDropOp;
 }
 
 #undef LOCTEXT_NAMESPACE
\ No newline at end of file
diff --git a/Engine/Source/Editor/DataLayerEditor/Private/DataLayer/DataLayerPropertyTypeCustomization.h b/Engine/Source/Editor/DataLayerEditor/Private/DataLayer/DataLayerPropertyTypeCustomization.h
--- a/Engine/Source/Editor/DataLayerEditor/Private/DataLayer/DataLayerPropertyTypeCustomization.h
+++ b/Engine/Source/Editor/DataLayerEditor/Private/DataLayer/DataLayerPropertyTypeCustomization.h
@@ -9,6 +9,7 @@ struct EVisibility;
 class FReply;
 class SWidget;
 class FDragDropOperation;
+class FDataLayerDragDropOp;
 class UDataLayer;
 
 struct FDataLayerPropertyTypeCustomization : public IPropertyTypeCustomization
@@ -28,6 +29,7 @@ private:
 	FReply OnSelectDataLayer();
 	FReply OnDrop(TSharedPtr<FDragDropOperation> InDragDrop);
 	bool OnVerifyDrag(TSharedPtr<FDragDropOperation> InDragDrop);
+	TSharedPtr<const FDataLayerDragDropOp> GetDataLayerDragDropOp(TSharedPtr<FDragDropOperation> InDragDrop);
 
 	TSharedPtr<IPropertyHandle> PropertyHandle;
 };
