Fixed Actor's Data Layer display name when multi section differs in details view.

diff --git a/Engine/Source/Editor/DataLayerEditor/Private/DataLayer/DataLayerPropertyTypeCustomization.cpp b/Engine/Source/Editor/DataLayerEditor/Private/DataLayer/DataLayerPropertyTypeCustomization.cpp
--- a/Engine/Source/Editor/DataLayerEditor/Private/DataLayer/DataLayerPropertyTypeCustomization.cpp
+++ b/Engine/Source/Editor/DataLayerEditor/Private/DataLayer/DataLayerPropertyTypeCustomization.cpp
@@ -91,10 +91,15 @@ void FDataLayerPropertyTypeCustomization::CustomizeHeader(TSharedRef<IPropertyHa
 	HeaderRow.IsEnabled(TAttribute<bool>(StructPropertyHandle, &IPropertyHandle::IsEditable));
 }
 
-UDataLayer* FDataLayerPropertyTypeCustomization::GetDataLayerFromPropertyHandle() const
+UDataLayer* FDataLayerPropertyTypeCustomization::GetDataLayerFromPropertyHandle(FPropertyAccess::Result* OutPropertyAccessResult) const
 {
 	FName DataLayerName;
-	if (PropertyHandle->GetValue(DataLayerName) == FPropertyAccess::Success)
+	FPropertyAccess::Result Result = PropertyHandle->GetValue(DataLayerName);
+	if (OutPropertyAccessResult)
+	{
+		*OutPropertyAccessResult = Result;
+	}
+	if (Result == FPropertyAccess::Success)
 	{
 		UDataLayer* DataLayer = UDataLayerEditorSubsystem::Get()->GetDataLayerFromName(DataLayerName);
 		return DataLayer;
@@ -104,7 +109,13 @@ UDataLayer* FDataLayerPropertyTypeCustomization::GetDataLayerFromPropertyHandle(
 
 FText FDataLayerPropertyTypeCustomization::GetDataLayerText() const
 {
-	return UDataLayer::GetDataLayerText(GetDataLayerFromPropertyHandle());
+	FPropertyAccess::Result PropertyAccessResult;
+	const UDataLayer* DataLayer = GetDataLayerFromPropertyHandle(&PropertyAccessResult);
+	if (PropertyAccessResult == FPropertyAccess::MultipleValues)
+	{
+		return NSLOCTEXT("PropertyEditor", "MultipleValues", "Multiple Values");
+	}
+	return UDataLayer::GetDataLayerText(DataLayer);
 }
 
 TSharedRef<SWidget> FDataLayerPropertyTypeCustomization::OnGetDataLayerMenu()
diff --git a/Engine/Source/Editor/DataLayerEditor/Private/DataLayer/DataLayerPropertyTypeCustomization.h b/Engine/Source/Editor/DataLayerEditor/Private/DataLayer/DataLayerPropertyTypeCustomization.h
--- a/Engine/Source/Editor/DataLayerEditor/Private/DataLayer/DataLayerPropertyTypeCustomization.h
+++ b/Engine/Source/Editor/DataLayerEditor/Private/DataLayer/DataLayerPropertyTypeCustomization.h
@@ -3,6 +3,7 @@
 #pragma once
 
 #include "CoreTypes.h"
+#include "PropertyEditorModule.h"
 #include "IPropertyTypeCustomization.h"
 
 struct EVisibility;
@@ -21,7 +22,7 @@ private:
 
 	void AssignDataLayer(const UDataLayer* InDataLayer);
 
-	UDataLayer* GetDataLayerFromPropertyHandle() const;
+	UDataLayer* GetDataLayerFromPropertyHandle(FPropertyAccess::Result* OutPropertyAccessResult = nullptr) const;
 	FText GetDataLayerText() const;
 	EVisibility GetSelectDataLayerVisibility() const;
 
