Fixed DataLayerHierarchy not properly uninitialized. Also fixed case where hierarchy could access deleted SDataLayerBrowser.

diff --git a/Engine/Source/Editor/DataLayerEditor/Private/DataLayer/DataLayerHierarchy.cpp b/Engine/Source/Editor/DataLayerEditor/Private/DataLayer/DataLayerHierarchy.cpp
--- a/Engine/Source/Editor/DataLayerEditor/Private/DataLayer/DataLayerHierarchy.cpp
+++ b/Engine/Source/Editor/DataLayerEditor/Private/DataLayer/DataLayerHierarchy.cpp
@@ -4,7 +4,6 @@
 #include "DataLayerMode.h"
 #include "DataLayerActorTreeItem.h"
 #include "DataLayerTreeItem.h"
-#include "SDataLayerBrowser.h"
 #include "DataLayerMode.h"
 #include "DataLayer/DataLayerEditorSubsystem.h"
 #include "WorldPartition/DataLayer/DataLayer.h"
@@ -14,29 +13,37 @@
 
 TUniquePtr<FDataLayerHierarchy> FDataLayerHierarchy::Create(FDataLayerMode* Mode, const TWeakObjectPtr<UWorld>& World)
 {
-	FDataLayerHierarchy* Hierarchy = new FDataLayerHierarchy(Mode, World);
-
-	GEngine->OnLevelActorAdded().AddRaw(Hierarchy, &FDataLayerHierarchy::OnLevelActorAdded);
-	GEngine->OnLevelActorDeleted().AddRaw(Hierarchy, &FDataLayerHierarchy::OnLevelActorDeleted);
-	GEngine->OnLevelActorListChanged().AddRaw(Hierarchy, &FDataLayerHierarchy::OnLevelActorListChanged);
-	UDataLayerEditorSubsystem::Get()->OnDataLayerChanged().AddRaw(Hierarchy, &FDataLayerHierarchy::OnDataLayerChanged);
-	UDataLayerEditorSubsystem::Get()->OnActorDataLayersChanged().AddRaw(Hierarchy, &FDataLayerHierarchy::OnActorDataLayersChanged);
-	Mode->GetDataLayerBrowser()->OnModeChanged().AddRaw(Hierarchy, &FDataLayerHierarchy::OnDataLayerBrowserModeChanged);
-
-	FWorldDelegates::LevelAddedToWorld.AddRaw(Hierarchy, &FDataLayerHierarchy::OnLevelAdded);
-	FWorldDelegates::LevelRemovedFromWorld.AddRaw(Hierarchy, &FDataLayerHierarchy::OnLevelRemoved);
-
-	return TUniquePtr<FDataLayerHierarchy>(Hierarchy);
+	return TUniquePtr<FDataLayerHierarchy>(new FDataLayerHierarchy(Mode, World));
 }
 
 FDataLayerHierarchy::FDataLayerHierarchy(FDataLayerMode* Mode, const TWeakObjectPtr<UWorld>& World)
 	: ISceneOutlinerHierarchy(Mode)
 	, RepresentingWorld(World)
+	, DataLayerBrowser(StaticCastSharedRef<SDataLayerBrowser>(Mode->GetDataLayerBrowser()->AsShared()))
 {
+	check(DataLayerBrowser.IsValid());
+	DataLayerBrowser.Pin()->OnModeChanged().AddRaw(this, &FDataLayerHierarchy::OnDataLayerBrowserModeChanged);
+
+	if (GEngine)
+	{
+		GEngine->OnLevelActorAdded().AddRaw(this, &FDataLayerHierarchy::OnLevelActorAdded);
+		GEngine->OnLevelActorDeleted().AddRaw(this, &FDataLayerHierarchy::OnLevelActorDeleted);
+		GEngine->OnLevelActorListChanged().AddRaw(this, &FDataLayerHierarchy::OnLevelActorListChanged);
+	}
+
+	UDataLayerEditorSubsystem::Get()->OnDataLayerChanged().AddRaw(this, &FDataLayerHierarchy::OnDataLayerChanged);
+	UDataLayerEditorSubsystem::Get()->OnActorDataLayersChanged().AddRaw(this, &FDataLayerHierarchy::OnActorDataLayersChanged);
+	FWorldDelegates::LevelAddedToWorld.AddRaw(this, &FDataLayerHierarchy::OnLevelAdded);
+	FWorldDelegates::LevelRemovedFromWorld.AddRaw(this, &FDataLayerHierarchy::OnLevelRemoved);
 }
 
 FDataLayerHierarchy::~FDataLayerHierarchy()
 {
+	if (DataLayerBrowser.IsValid())
+	{
+		DataLayerBrowser.Pin()->OnModeChanged().RemoveAll(this);
+	}
+
 	if (GEngine)
 	{
 		GEngine->OnLevelActorAdded().RemoveAll(this);
@@ -46,21 +53,10 @@ FDataLayerHierarchy::~FDataLayerHierarchy()
 
 	UDataLayerEditorSubsystem::Get()->OnDataLayerChanged().RemoveAll(this);
 	UDataLayerEditorSubsystem::Get()->OnActorDataLayersChanged().RemoveAll(this);
-
-	if (SDataLayerBrowser* DataLayerBrowser = GetDataLayerMode() ? GetDataLayerMode()->GetDataLayerBrowser() : nullptr)
-	{
-		DataLayerBrowser->OnModeChanged().RemoveAll(this);
-	}
-
 	FWorldDelegates::LevelAddedToWorld.RemoveAll(this);
 	FWorldDelegates::LevelRemovedFromWorld.RemoveAll(this);
 }
 
-FDataLayerMode* FDataLayerHierarchy::GetDataLayerMode() const
-{
-	return (FDataLayerMode*)Mode;
-}
-
 void FDataLayerHierarchy::CreateItems(TArray<FSceneOutlinerTreeItemPtr>& OutItems) const
 {
 	const AWorldDataLayers* WorldDataLayers = AWorldDataLayers::Get(RepresentingWorld.Get());
@@ -70,15 +66,15 @@ void FDataLayerHierarchy::CreateItems(TArray<FSceneOutlinerTreeItemPtr>& OutItem
 	}
 
 	WorldDataLayers->ForEachDataLayer([this, &OutItems](UDataLayer* DataLayer)
+	{
+		if (FSceneOutlinerTreeItemPtr DataLayerItem = Mode->CreateItemFor<FDataLayerTreeItem>(DataLayer))
 		{
-			if (FSceneOutlinerTreeItemPtr DataLayerItem = Mode->CreateItemFor<FDataLayerTreeItem>(DataLayer))
-			{
-				OutItems.Add(DataLayerItem);
-			}
-			return true;
-		});
+			OutItems.Add(DataLayerItem);
+		}
+		return true;
+	});
 
-	if (GetDataLayerMode()->GetDataLayerBrowser()->GetMode() == EDataLayerBrowserMode::DataLayerContents)
+	if (DataLayerBrowser.IsValid() && DataLayerBrowser.Pin()->GetMode() == EDataLayerBrowserMode::DataLayerContents)
 	{
 		for (AActor* Actor : FActorRange(RepresentingWorld.Get()))
 		{
@@ -220,4 +216,4 @@ void FDataLayerHierarchy::FullRefreshEvent()
 	FSceneOutlinerHierarchyChangedData EventData;
 	EventData.Type = FSceneOutlinerHierarchyChangedData::FullRefresh;
 	HierarchyChangedEvent.Broadcast(EventData);
-}
\ No newline at end of file
+}
diff --git a/Engine/Source/Editor/DataLayerEditor/Private/DataLayer/DataLayerHierarchy.h b/Engine/Source/Editor/DataLayerEditor/Private/DataLayer/DataLayerHierarchy.h
--- a/Engine/Source/Editor/DataLayerEditor/Private/DataLayer/DataLayerHierarchy.h
+++ b/Engine/Source/Editor/DataLayerEditor/Private/DataLayer/DataLayerHierarchy.h
@@ -36,7 +36,7 @@ private:
 	void OnDataLayerChanged(const EDataLayerAction Action, const TWeakObjectPtr<const UDataLayer>& ChangedDataLayer, const FName& ChangedProperty);
 	void OnDataLayerBrowserModeChanged(EDataLayerBrowserMode InMode);
 	void FullRefreshEvent();
-	FDataLayerMode* GetDataLayerMode() const;
 
 	TWeakObjectPtr<UWorld> RepresentingWorld;
+	TWeakPtr<SDataLayerBrowser> DataLayerBrowser;
 };
\ No newline at end of file
diff --git a/Engine/Source/Editor/DataLayerEditor/Private/DataLayer/DataLayerMode.cpp b/Engine/Source/Editor/DataLayerEditor/Private/DataLayer/DataLayerMode.cpp
--- a/Engine/Source/Editor/DataLayerEditor/Private/DataLayer/DataLayerMode.cpp
+++ b/Engine/Source/Editor/DataLayerEditor/Private/DataLayer/DataLayerMode.cpp
@@ -43,11 +43,6 @@ FDataLayerMode::FDataLayerMode(const FDataLayerModeParams& Params)
 	const_cast<FSharedSceneOutlinerData&>(SceneOutliner->GetSharedData()).CustomDelete = FCustomSceneOutlinerDeleteDelegate::CreateRaw(this, &FDataLayerMode::DeleteItems);
 }
 
-FDataLayerMode::~FDataLayerMode()
-{
-	Hierarchy.Release();
-}
-
 int32 FDataLayerMode::GetTypeSortPriority(const ISceneOutlinerTreeItem& Item) const
 {
 	if (Item.IsA<FDataLayerTreeItem>())
@@ -786,4 +781,4 @@ void FDataLayerMode::ChooseRepresentingWorld()
 	}
 }
 
-#undef LOCTEXT_NAMESPACE
\ No newline at end of file
+#undef LOCTEXT_NAMESPACE
diff --git a/Engine/Source/Editor/DataLayerEditor/Private/DataLayer/DataLayerMode.h b/Engine/Source/Editor/DataLayerEditor/Private/DataLayer/DataLayerMode.h
--- a/Engine/Source/Editor/DataLayerEditor/Private/DataLayer/DataLayerMode.h
+++ b/Engine/Source/Editor/DataLayerEditor/Private/DataLayer/DataLayerMode.h
@@ -32,7 +32,6 @@ public:
 	};
 
 	FDataLayerMode(const FDataLayerModeParams& Params);
-	virtual ~FDataLayerMode();
 
 	virtual void Rebuild() override;
 	virtual TSharedPtr<SWidget> CreateContextMenu() override;
