DataLayers: Show unloaded actors in data layers view

diff --git a/Engine/Source/Editor/DataLayerEditor/Private/DataLayer/DataLayerHierarchy.cpp b/Engine/Source/Editor/DataLayerEditor/Private/DataLayer/DataLayerHierarchy.cpp
--- a/Engine/Source/Editor/DataLayerEditor/Private/DataLayer/DataLayerHierarchy.cpp
+++ b/Engine/Source/Editor/DataLayerEditor/Private/DataLayer/DataLayerHierarchy.cpp
@@ -3,6 +3,7 @@
 #include "DataLayerHierarchy.h"
 #include "DataLayerMode.h"
 #include "DataLayerActorTreeItem.h"
+#include "DataLayersActorDescTreeItem.h"
 #include "DataLayerTreeItem.h"
 #include "DataLayerMode.h"
 #include "DataLayer/DataLayerEditorSubsystem.h"
@@ -10,6 +11,7 @@
 #include "Engine/World.h"
 #include "EngineUtils.h"
 #include "Editor.h"
+#include "WorldPartition/WorldPartitionHelpers.h"
 
 TUniquePtr<FDataLayerHierarchy> FDataLayerHierarchy::Create(FDataLayerMode* Mode, const TWeakObjectPtr<UWorld>& World)
 {
@@ -31,6 +33,21 @@ FDataLayerHierarchy::FDataLayerHierarchy(FDataLayerMode* Mode, const TWeakObject
 		GEngine->OnLevelActorListChanged().AddRaw(this, &FDataLayerHierarchy::OnLevelActorListChanged);
 	}
 
+	if (World.IsValid())
+	{
+		if (World->PersistentLevel)
+		{
+			World->PersistentLevel->OnLoadedActorAddedToLevelEvent.AddRaw(this, &FDataLayerHierarchy::OnLoadedActorAdded);
+			World->PersistentLevel->OnLoadedActorRemovedFromLevelEvent.AddRaw(this, &FDataLayerHierarchy::OnLoadedActorRemoved);
+		}
+		
+		if (UWorldPartition* WorldPartition = World->GetWorldPartition())
+		{
+			WorldPartition->OnActorDescAddedEvent.AddRaw(this, &FDataLayerHierarchy::OnActorDescAdded);
+			WorldPartition->OnActorDescRemovedEvent.AddRaw(this, &FDataLayerHierarchy::OnActorDescRemoved);
+		}
+	}
+
 	UDataLayerEditorSubsystem::Get()->OnDataLayerChanged().AddRaw(this, &FDataLayerHierarchy::OnDataLayerChanged);
 	UDataLayerEditorSubsystem::Get()->OnActorDataLayersChanged().AddRaw(this, &FDataLayerHierarchy::OnActorDataLayersChanged);
 	FWorldDelegates::LevelAddedToWorld.AddRaw(this, &FDataLayerHierarchy::OnLevelAdded);
@@ -51,6 +68,21 @@ FDataLayerHierarchy::~FDataLayerHierarchy()
 		GEngine->OnLevelActorListChanged().RemoveAll(this);
 	}
 
+	if (RepresentingWorld.IsValid())
+	{
+		if (RepresentingWorld->PersistentLevel)
+		{
+			RepresentingWorld->PersistentLevel->OnLoadedActorAddedToLevelEvent.RemoveAll(this);
+			RepresentingWorld->PersistentLevel->OnLoadedActorRemovedFromLevelEvent.RemoveAll(this);
+		}
+		
+		if (UWorldPartition* WorldPartition = RepresentingWorld->GetWorldPartition())
+		{
+			WorldPartition->OnActorDescAddedEvent.RemoveAll(this);
+			WorldPartition->OnActorDescRemovedEvent.RemoveAll(this);
+		}
+	}
+
 	UDataLayerEditorSubsystem::Get()->OnDataLayerChanged().RemoveAll(this);
 	UDataLayerEditorSubsystem::Get()->OnActorDataLayersChanged().RemoveAll(this);
 	FWorldDelegates::LevelAddedToWorld.RemoveAll(this);
@@ -89,6 +121,27 @@ void FDataLayerHierarchy::CreateItems(TArray<FSceneOutlinerTreeItemPtr>& OutItem
 				}
 			}
 		}
+
+		if (UWorldPartition* const WorldPartition = RepresentingWorld.Get()->GetWorldPartition())
+		{
+			FWorldPartitionHelpers::ForEachActorDesc(WorldPartition, [this, WorldPartition, WorldDataLayers, &OutItems](const FWorldPartitionActorDesc* ActorDesc)
+				{
+					if (ActorDesc != nullptr && !ActorDesc->IsLoaded())
+					{
+						for (const FName& DataLayerName : ActorDesc->GetDataLayers())
+						{
+							if (const UDataLayer* const DataLayer = WorldDataLayers->GetDataLayerFromName(DataLayerName))
+							{
+								if (const FSceneOutlinerTreeItemPtr ActorDescItem = Mode->CreateItemFor<FDataLayerActorDescTreeItem>(FDataLayerActorDescTreeItemData(ActorDesc->GetGuid(), WorldPartition, const_cast<UDataLayer*>(DataLayer))))
+								{
+									OutItems.Add(ActorDescItem);
+								}
+							}
+						}
+					}
+					return true;
+				});
+		}
 	}
 }
 
@@ -108,6 +161,16 @@ FSceneOutlinerTreeItemPtr FDataLayerHierarchy::FindParent(const ISceneOutlinerTr
 			}
 		}
 	}
+	else if (const FDataLayerActorDescTreeItem* DataLayerActorDescTreeItem = Item.CastTo<FDataLayerActorDescTreeItem>())
+	{
+		if (const UDataLayer* DataLayer = DataLayerActorDescTreeItem->GetDataLayer())
+		{
+			if (const FSceneOutlinerTreeItemPtr* ParentItem = Items.Find(DataLayer))
+			{
+				return *ParentItem;
+			}
+		}
+	}
 	return nullptr;
 }
 
@@ -121,7 +184,14 @@ FSceneOutlinerTreeItemPtr FDataLayerHierarchy::CreateParentItem(const FSceneOutl
 	{
 		if (UDataLayer* DataLayer = DataLayerActorTreeItem->GetDataLayer())
 		{
-			return Mode->CreateItemFor<FDataLayerTreeItem>(DataLayer);
+			return Mode->CreateItemFor<FDataLayerTreeItem>(DataLayer, true);
+		}
+	}
+	else if (FDataLayerActorDescTreeItem* DataLayerActorDescTreeItem = Item->CastTo<FDataLayerActorDescTreeItem>())
+	{
+		if (UDataLayer* DataLayer = DataLayerActorDescTreeItem->GetDataLayer())
+		{
+			return Mode->CreateItemFor<FDataLayerTreeItem>(DataLayer, true);
 		}
 	}
 	return nullptr;
@@ -179,7 +249,7 @@ void FDataLayerHierarchy::OnLevelActorDeleted(AActor* InActor)
 			FSceneOutlinerHierarchyChangedData EventData;
 			EventData.Type = FSceneOutlinerHierarchyChangedData::Removed;
 
-			TArray<const UDataLayer*> DataLayers = InActor->GetDataLayerObjects();
+			const TArray<const UDataLayer*> DataLayers = InActor->GetDataLayerObjects();
 			EventData.ItemIDs.Reserve(DataLayers.Num());
 			for (const UDataLayer* DataLayer : DataLayers)
 			{
@@ -211,6 +281,100 @@ void FDataLayerHierarchy::OnLevelRemoved(ULevel* InLevel, UWorld* InWorld)
 	}
 }
 
+void FDataLayerHierarchy::OnLoadedActorAdded(AActor& InActor)
+{
+	// Handle the actor being added to the level
+	OnLevelActorAdded(&InActor);
+
+	// Remove corresponding actor desc items
+	if (RepresentingWorld.Get() == InActor.GetWorld() && InActor.HasDataLayers())
+	{
+		FSceneOutlinerHierarchyChangedData EventData;
+		EventData.Type = FSceneOutlinerHierarchyChangedData::Removed;
+		
+		const TArray<const UDataLayer*> DataLayers = InActor.GetDataLayerObjects();
+		EventData.ItemIDs.Reserve(DataLayers.Num());
+		for (const UDataLayer* DataLayer : DataLayers)
+		{
+			EventData.ItemIDs.Add(FDataLayerActorDescTreeItem::ComputeTreeItemID(InActor.GetActorGuid(), DataLayer));
+		}
+		HierarchyChangedEvent.Broadcast(EventData);
+	}
+}
+
+void FDataLayerHierarchy::OnLoadedActorRemoved(AActor& InActor)
+{
+	// Handle the actor being removed from the level
+	OnLevelActorDeleted(&InActor);
+
+	// Add any corresponding actor desc items for this actor
+	UWorldPartition* const WorldPartition = RepresentingWorld.Get()->GetWorldPartition();
+	if (RepresentingWorld.Get() == InActor.GetWorld() && InActor.HasDataLayers() && WorldPartition != nullptr)
+	{
+		FSceneOutlinerHierarchyChangedData EventData;
+		EventData.Type = FSceneOutlinerHierarchyChangedData::Added;
+
+		const TArray<const UDataLayer*> DataLayers = InActor.GetDataLayerObjects();
+		EventData.Items.Reserve(DataLayers.Num());
+		for (const UDataLayer* DataLayer : DataLayers)
+		{
+			EventData.Items.Add(Mode->CreateItemFor<FDataLayerActorDescTreeItem>(FDataLayerActorDescTreeItemData(InActor.GetActorGuid(), WorldPartition, const_cast<UDataLayer*>(DataLayer))));
+		}
+		HierarchyChangedEvent.Broadcast(EventData);
+    }
+}
+
+void FDataLayerHierarchy::OnActorDescAdded(FWorldPartitionActorDesc* InActorDesc)
+{
+	if (InActorDesc == nullptr)
+	{
+		return;
+	}
+
+	UWorldPartition* const WorldPartition = RepresentingWorld.Get()->GetWorldPartition();
+	const AWorldDataLayers* const WorldDataLayers = RepresentingWorld.Get()->GetWorldDataLayers();
+	const TArray<FName>& DataLayerNames = InActorDesc->GetDataLayers();
+
+	if (WorldDataLayers != nullptr && RepresentingWorld->GetWorldPartition() == WorldPartition && DataLayerNames.Num() > 0)
+	{
+		FSceneOutlinerHierarchyChangedData EventData;
+		EventData.Type = FSceneOutlinerHierarchyChangedData::Added;
+
+		EventData.Items.Reserve(DataLayerNames.Num());
+		for (const FName& DataLayerName : DataLayerNames)
+		{
+			const UDataLayer* const DataLayer = WorldDataLayers->GetDataLayerFromName(DataLayerName);
+			EventData.Items.Add(Mode->CreateItemFor<FDataLayerActorDescTreeItem>(FDataLayerActorDescTreeItemData(InActorDesc->GetGuid(), WorldPartition, const_cast<UDataLayer*>(DataLayer))));
+		}
+		HierarchyChangedEvent.Broadcast(EventData);
+	}
+}
+
+void FDataLayerHierarchy::OnActorDescRemoved(FWorldPartitionActorDesc* InActorDesc)
+{
+	if (InActorDesc == nullptr)
+	{
+		return;
+	}
+
+	const AWorldDataLayers* const WorldDataLayers = RepresentingWorld.Get()->GetWorldDataLayers();
+	const TArray<FName>& DataLayerNames = InActorDesc->GetDataLayers();
+	
+	if (WorldDataLayers != nullptr && DataLayerNames.Num() > 0)
+	{
+		FSceneOutlinerHierarchyChangedData EventData;
+		EventData.Type = FSceneOutlinerHierarchyChangedData::Removed;
+		EventData.ItemIDs.Reserve(DataLayerNames.Num());
+
+		for (const FName& DataLayerName : DataLayerNames)
+		{
+			const UDataLayer* const DataLayer = WorldDataLayers->GetDataLayerFromName(DataLayerName);
+			EventData.ItemIDs.Add(FDataLayerActorDescTreeItem::ComputeTreeItemID(InActorDesc->GetGuid(), DataLayer));
+		}
+		HierarchyChangedEvent.Broadcast(EventData);
+	}
+}
+
 void FDataLayerHierarchy::FullRefreshEvent()
 {
 	FSceneOutlinerHierarchyChangedData EventData;
diff --git a/Engine/Source/Editor/DataLayerEditor/Private/DataLayer/DataLayerHierarchy.h b/Engine/Source/Editor/DataLayerEditor/Private/DataLayer/DataLayerHierarchy.h
--- a/Engine/Source/Editor/DataLayerEditor/Private/DataLayer/DataLayerHierarchy.h
+++ b/Engine/Source/Editor/DataLayerEditor/Private/DataLayer/DataLayerHierarchy.h
@@ -32,6 +32,10 @@ private:
 	void OnLevelActorListChanged();
 	void OnLevelAdded(ULevel* InLevel, UWorld* InWorld);
 	void OnLevelRemoved(ULevel* InLevel, UWorld* InWorld);
+	void OnLoadedActorAdded(AActor& InActor);
+	void OnLoadedActorRemoved(AActor& InActor);
+	void OnActorDescAdded(FWorldPartitionActorDesc* InActorDesc);
+	void OnActorDescRemoved(FWorldPartitionActorDesc* InActorDesc);
 	void OnActorDataLayersChanged(const TWeakObjectPtr<AActor>& InActor);
 	void OnDataLayerChanged(const EDataLayerAction Action, const TWeakObjectPtr<const UDataLayer>& ChangedDataLayer, const FName& ChangedProperty);
 	void OnDataLayerBrowserModeChanged(EDataLayerBrowserMode InMode);
diff --git a/Engine/Source/Editor/DataLayerEditor/Private/DataLayer/DataLayerMode.cpp b/Engine/Source/Editor/DataLayerEditor/Private/DataLayer/DataLayerMode.cpp
--- a/Engine/Source/Editor/DataLayerEditor/Private/DataLayer/DataLayerMode.cpp
+++ b/Engine/Source/Editor/DataLayerEditor/Private/DataLayer/DataLayerMode.cpp
@@ -14,6 +14,7 @@
 #include "DragAndDrop/ActorDragDropOp.h"
 #include "DragAndDrop/FolderDragDropOp.h"
 #include "ActorMode.h"
+#include "DataLayersActorDescTreeItem.h"
 #include "EditorActorFolders.h"
 #include "Algo/Transform.h"
 #include "ToolMenus.h"
@@ -47,11 +48,15 @@ int32 FDataLayerMode::GetTypeSortPriority(const ISceneOutlinerTreeItem& Item) co
 {
 	if (Item.IsA<FDataLayerTreeItem>())
 	{
-		return (int32)EItemSortOrder::DataLayer;
+		return static_cast<int32>(EItemSortOrder::DataLayer);
 	}
 	else if (Item.IsA<FDataLayerActorTreeItem>())
 	{
-		return (int32)EItemSortOrder::Actor;
+		return static_cast<int32>(EItemSortOrder::Actor);
+	}
+	else if (Item.IsA<FDataLayerActorDescTreeItem>())
+	{
+		return static_cast<int32_t>(EItemSortOrder::Unloaded);
 	}
 	// Warning: using actor mode with an unsupported item type!
 	check(false);
diff --git a/Engine/Source/Editor/DataLayerEditor/Private/DataLayer/DataLayerMode.h b/Engine/Source/Editor/DataLayerEditor/Private/DataLayer/DataLayerMode.h
--- a/Engine/Source/Editor/DataLayerEditor/Private/DataLayer/DataLayerMode.h
+++ b/Engine/Source/Editor/DataLayerEditor/Private/DataLayer/DataLayerMode.h
@@ -28,7 +28,8 @@ public:
 	enum class EItemSortOrder : int32
 	{
 		DataLayer = 0,
-		Actor = 10
+		Actor = 10,
+		Unloaded = 20,
 	};
 
 	FDataLayerMode(const FDataLayerModeParams& Params);
diff --git a/Engine/Source/Editor/DataLayerEditor/Private/DataLayer/DataLayersActorDescTreeItem.h b/Engine/Source/Editor/DataLayerEditor/Private/DataLayer/DataLayersActorDescTreeItem.h
new file mode 100644
--- /dev/null
+++ b/Engine/Source/Editor/DataLayerEditor/Private/DataLayer/DataLayersActorDescTreeItem.h
@@ -0,0 +1,74 @@
+// Copyright Epic Games, Inc. All Rights Reserved.
+
+#pragma once
+
+#include "CoreMinimal.h"
+#include "SceneOutlinerFwd.h"
+#include "ActorDescTreeItem.h"
+#include "WorldPartition/DataLayer/DataLayer.h"
+
+//////////////////////////////////////////////////////////////////////////
+// FDataLayerActorTreeItemData
+
+class FWorldPartitionActorDesc;
+
+struct FDataLayerActorDescTreeItemData
+{
+	FDataLayerActorDescTreeItemData(const FGuid& InActorGuid, UActorDescContainer* InContainer, UDataLayer* InDataLayer)
+		: ActorGuid(InActorGuid)
+		, Container(InContainer)
+		, DataLayer(InDataLayer)
+	{}
+	
+	const FGuid& ActorGuid;
+	UActorDescContainer* const Container;
+	TWeakObjectPtr<UDataLayer> DataLayer;
+};
+
+//////////////////////////////////////////////////////////////////////////
+// FDataLayerActorTreeItem
+
+struct FDataLayerActorDescTreeItem : public FActorDescTreeItem
+{
+public:
+	DECLARE_DELEGATE_RetVal_TwoParams(bool, FFilterPredicate, const FWorldPartitionActorDesc* ActorDesc, const UDataLayer* DataLayer);
+	DECLARE_DELEGATE_RetVal_TwoParams(bool, FInteractivePredicate, const FWorldPartitionActorDesc* ActorDesc, const UDataLayer* DataLayer);
+
+	FDataLayerActorDescTreeItem(const FDataLayerActorDescTreeItemData& InData)
+		: FActorDescTreeItem(InData.ActorGuid, InData.Container)
+		, DataLayer(InData.DataLayer)
+		, IDDataLayerActorDesc(FDataLayerActorDescTreeItem::ComputeTreeItemID(InData.ActorGuid, DataLayer.Get()))
+	{
+	}
+
+	const UDataLayer* GetDataLayer() const { return DataLayer.Get(); }
+	UDataLayer* GetDataLayer() { return DataLayer.Get(); }
+
+	static uint32 ComputeTreeItemID(const FGuid& InActorGuid, const UDataLayer* InDataLayer)
+	{
+		return HashCombine(GetTypeHash(InActorGuid), GetTypeHash(FObjectKey(InDataLayer)));
+	}
+
+	bool Filter(FFilterPredicate Pred) const
+	{
+		return Pred.Execute(ActorDescHandle.GetActorDesc(), DataLayer.Get());
+	}
+
+	bool GetInteractiveState(FInteractivePredicate Pred) const
+	{
+		return Pred.Execute(ActorDescHandle.GetActorDesc(), DataLayer.Get());
+	}
+
+	/* Begin ISceneOutlinerTreeItem Implementation */
+	virtual bool IsValid() const override { return ActorDescHandle.IsValid() && DataLayer.IsValid(); }
+	virtual FSceneOutlinerTreeItemID GetID() const override { return IDDataLayerActorDesc; }
+	virtual bool ShouldShowVisibilityState() const { return false; }
+	virtual bool HasVisibilityInfo() const override { return false; }
+	virtual void OnVisibilityChanged(const bool bNewVisibility) override {}
+	virtual bool GetVisibility() const override { return false; }
+	/* End ISceneOutlinerTreeItem Implementation */
+
+private:
+	TWeakObjectPtr<UDataLayer> DataLayer;
+	const uint32 IDDataLayerActorDesc;
+};
\ No newline at end of file
