Horde Issue 69162 - Fixed non-unity build

diff --git a/Engine/Source/Editor/DataLayerEditor/Private/DataLayer/DataLayerHierarchy.h b/Engine/Source/Editor/DataLayerEditor/Private/DataLayer/DataLayerHierarchy.h
--- a/Engine/Source/Editor/DataLayerEditor/Private/DataLayer/DataLayerHierarchy.h
+++ b/Engine/Source/Editor/DataLayerEditor/Private/DataLayer/DataLayerHierarchy.h
@@ -11,6 +11,7 @@
 class FDataLayerMode;
 class UDataLayer;
 class UWorld;
+class FWorldPartitionActorDesc;
 
 class FDataLayerHierarchy : public ISceneOutlinerHierarchy
 {
