DataLayerBrowser: Optimize LevelAdded/Removed to avoid FullRefresh (from ~1000 calls costing ~800ms to ~20ms)

diff --git a/Engine/Source/Editor/DataLayerEditor/Private/DataLayer/DataLayerHierarchy.cpp b/Engine/Source/Editor/DataLayerEditor/Private/DataLayer/DataLayerHierarchy.cpp
--- a/Engine/Source/Editor/DataLayerEditor/Private/DataLayer/DataLayerHierarchy.cpp
+++ b/Engine/Source/Editor/DataLayerEditor/Private/DataLayer/DataLayerHierarchy.cpp
@@ -197,26 +197,69 @@ FSceneOutlinerTreeItemPtr FDataLayerHierarchy::CreateParentItem(const FSceneOutl
 	return nullptr;
 }
 
-void FDataLayerHierarchy::OnLevelActorAdded(AActor* InActor)
+void FDataLayerHierarchy::OnLevelActorsAdded(const TArray<AActor*>& InActors)
 {
-	if (InActor && RepresentingWorld.Get() == InActor->GetWorld())
+	if (UWorld* CurrentWorld = RepresentingWorld.Get())
 	{
-		if (InActor->HasDataLayers())
+		FSceneOutlinerHierarchyChangedData EventData;
+		EventData.Type = FSceneOutlinerHierarchyChangedData::Added;
+
+		for (AActor* Actor : InActors)
 		{
-			FSceneOutlinerHierarchyChangedData EventData;
-			EventData.Type = FSceneOutlinerHierarchyChangedData::Added;
+			if (Actor != nullptr && Actor->HasDataLayers() && Actor->GetWorld() == CurrentWorld)
+			{
+				TArray<const UDataLayer*> DataLayers = Actor->GetDataLayerObjects();
+				EventData.Items.Reserve(DataLayers.Num());
+				for (const UDataLayer* DataLayer : DataLayers)
+				{
+					EventData.Items.Add(Mode->CreateItemFor<FDataLayerActorTreeItem>(FDataLayerActorTreeItemData(Actor, const_cast<UDataLayer*>(DataLayer))));
+				}
+			}
+		}
 
-			TArray<const UDataLayer*> DataLayers = InActor->GetDataLayerObjects();
-			EventData.Items.Reserve(DataLayers.Num());
-			for (const UDataLayer* DataLayer : DataLayers)
+		if (!EventData.Items.IsEmpty())
+		{
+			HierarchyChangedEvent.Broadcast(EventData);
+		}
+	}
+}
+
+void FDataLayerHierarchy::OnLevelActorsRemoved(const TArray<AActor*>& InActors)
+{
+	if(UWorld* CurrentWorld = RepresentingWorld.Get())
+	{
+		FSceneOutlinerHierarchyChangedData EventData;
+		EventData.Type = FSceneOutlinerHierarchyChangedData::Removed;
+
+		if (AWorldDataLayers* WorldDataLayers = CurrentWorld->GetWorldDataLayers())
+		{
+			for (AActor* Actor : InActors)
 			{
-				EventData.Items.Add(Mode->CreateItemFor<FDataLayerActorTreeItem>(FDataLayerActorTreeItemData(InActor, const_cast<UDataLayer*>(DataLayer))));
+				if (Actor != nullptr && Actor->HasDataLayers())
+				{
+					// It is possible here that Actor doesn't have world anymore
+					const TArray<const UDataLayer*> DataLayers = Actor->GetDataLayerObjects(WorldDataLayers);
+					EventData.ItemIDs.Reserve(DataLayers.Num());
+					for (const UDataLayer* DataLayer : DataLayers)
+					{
+						EventData.ItemIDs.Add(FDataLayerActorTreeItem::ComputeTreeItemID(Actor, DataLayer));
+					}
+				}
 			}
+		}
+
+		if (!EventData.ItemIDs.IsEmpty())
+		{
 			HierarchyChangedEvent.Broadcast(EventData);
 		}
 	}
 }
 
+void FDataLayerHierarchy::OnLevelActorAdded(AActor* InActor)
+{
+	OnLevelActorsAdded({ InActor });
+}
+
 void FDataLayerHierarchy::OnActorDataLayersChanged(const TWeakObjectPtr<AActor>& InActor)
 {
 	AActor* Actor = InActor.Get();
@@ -242,22 +285,7 @@ void FDataLayerHierarchy::OnDataLayerBrowserModeChanged(EDataLayerBrowserMode In
 
 void FDataLayerHierarchy::OnLevelActorDeleted(AActor* InActor)
 {
-	if (RepresentingWorld.Get() == InActor->GetWorld())
-	{
-		if (InActor->HasDataLayers())
-		{
-			FSceneOutlinerHierarchyChangedData EventData;
-			EventData.Type = FSceneOutlinerHierarchyChangedData::Removed;
-
-			const TArray<const UDataLayer*> DataLayers = InActor->GetDataLayerObjects();
-			EventData.ItemIDs.Reserve(DataLayers.Num());
-			for (const UDataLayer* DataLayer : DataLayers)
-			{
-				EventData.ItemIDs.Add(FDataLayerActorTreeItem::ComputeTreeItemID(InActor, DataLayer));
-			}
-			HierarchyChangedEvent.Broadcast(EventData);
-		}
-	}
+	OnLevelActorsRemoved({ InActor });
 }
 
 void FDataLayerHierarchy::OnLevelActorListChanged()
@@ -267,17 +295,17 @@ void FDataLayerHierarchy::OnLevelActorListChanged()
 
 void FDataLayerHierarchy::OnLevelAdded(ULevel* InLevel, UWorld* InWorld)
 {
-	if (RepresentingWorld.Get() == InWorld)
+	if (InLevel != nullptr && RepresentingWorld.Get() == InWorld)
 	{
-		FullRefreshEvent();
+		OnLevelActorsAdded(InLevel->Actors);
 	}
 }
 
 void FDataLayerHierarchy::OnLevelRemoved(ULevel* InLevel, UWorld* InWorld)
 {
-	if (RepresentingWorld.Get() == InWorld)
+	if (InLevel != nullptr && RepresentingWorld.Get() == InWorld)
 	{
-		FullRefreshEvent();
+		OnLevelActorsRemoved(InLevel->Actors);
 	}
 }
 
diff --git a/Engine/Source/Editor/DataLayerEditor/Private/DataLayer/DataLayerHierarchy.h b/Engine/Source/Editor/DataLayerEditor/Private/DataLayer/DataLayerHierarchy.h
--- a/Engine/Source/Editor/DataLayerEditor/Private/DataLayer/DataLayerHierarchy.h
+++ b/Engine/Source/Editor/DataLayerEditor/Private/DataLayer/DataLayerHierarchy.h
@@ -28,6 +28,8 @@ private:
 	FDataLayerHierarchy(const FDataLayerHierarchy&) = delete;
 	FDataLayerHierarchy& operator=(const FDataLayerHierarchy&) = delete;
 
+	void OnLevelActorsAdded(const TArray<AActor*>& InActors);
+	void OnLevelActorsRemoved(const TArray<AActor*>& InActors);
 	void OnLevelActorAdded(AActor* InActor);
 	void OnLevelActorDeleted(AActor* InActor);
 	void OnLevelActorListChanged();
diff --git a/Engine/Source/Runtime/Engine/Classes/GameFramework/Actor.h b/Engine/Source/Runtime/Engine/Classes/GameFramework/Actor.h
--- a/Engine/Source/Runtime/Engine/Classes/GameFramework/Actor.h
+++ b/Engine/Source/Runtime/Engine/Classes/GameFramework/Actor.h
@@ -39,6 +39,7 @@ struct FAttachedActorInfo;
 struct FNetViewer;
 struct FNetworkObjectInfo;
 class UDataLayer;
+class AWorldDataLayers;
 
 /** Chooses a method for actors to update overlap state (objects it is touching) on initialization, currently only used during level streaming. */
 UENUM(BlueprintType)
@@ -1081,6 +1082,7 @@ public:
 	bool HasAnyOfDataLayers(const TArray<FName>& DataLayerNames) const;
 	TArray<FName> GetDataLayerNames() const;
 	TArray<const UDataLayer*> GetDataLayerObjects() const;
+	TArray<const UDataLayer*> GetDataLayerObjects(const AWorldDataLayers* WorldDataLayers) const;
 	bool IsPropertyChangedAffectingDataLayers(FPropertyChangedEvent& PropertyChangedEvent) const;
 	bool IsValidForDataLayer() const;
 	void FixupDataLayers();
diff --git a/Engine/Source/Runtime/Engine/Private/ActorEditor.cpp b/Engine/Source/Runtime/Engine/Private/ActorEditor.cpp
--- a/Engine/Source/Runtime/Engine/Private/ActorEditor.cpp
+++ b/Engine/Source/Runtime/Engine/Private/ActorEditor.cpp
@@ -1158,7 +1158,11 @@ TArray<FName> AActor::GetDataLayerNames() const
 
 TArray<const UDataLayer*> AActor::GetDataLayerObjects() const
 {
-	const AWorldDataLayers* WorldDataLayers = GetWorld()->GetWorldDataLayers();
+	return GetDataLayerObjects(GetWorld()->GetWorldDataLayers());
+}
+
+TArray<const UDataLayer*> AActor::GetDataLayerObjects(const AWorldDataLayers* WorldDataLayers) const
+{
 	return WorldDataLayers ? WorldDataLayers->GetDataLayerObjects(DataLayers) : TArray<const UDataLayer*>();
 }
 
