Remove the ability to create new data layer by just drag & drop. Creating a data layer should not be necessarily easy to do and we expect a full project to only have a few. Current impl was confusing because it worked only for the 1st data layer that you created.

diff --git a/Engine/Source/Editor/DataLayerEditor/Private/DataLayer/DataLayerMode.cpp b/Engine/Source/Editor/DataLayerEditor/Private/DataLayer/DataLayerMode.cpp
--- a/Engine/Source/Editor/DataLayerEditor/Private/DataLayer/DataLayerMode.cpp
+++ b/Engine/Source/Editor/DataLayerEditor/Private/DataLayer/DataLayerMode.cpp
@@ -282,14 +282,6 @@ FSceneOutlinerDragValidationInfo FDataLayerMode::ValidateDrop(const ISceneOutlin
 
 			return FSceneOutlinerDragValidationInfo(ESceneOutlinerDropCompatibility::Compatible, FText::Format(LOCTEXT("AssignToDataLayer", "Assign to Data Layer \"{0}\""), FText::FromName(DataLayerTarget->GetDataLayerLabel())));
 		}
-		else
-		{
-			if (!PayloadActors[0]->HasDataLayers())
-			{
-				// Only allow actors not coming from the DataLayerBrowser
-				return FSceneOutlinerDragValidationInfo(ESceneOutlinerDropCompatibility::Compatible, LOCTEXT("AssignToNewDataLayer", "Assign to New Data Layer"));
-			}
-		}
 	}
 	return FSceneOutlinerDragValidationInfo::Invalid();
 }
@@ -394,18 +386,6 @@ void FDataLayerMode::OnDrop(ISceneOutlinerTreeItem& DropTarget, const FSceneOutl
 			DataLayerEditorSubsystem->AddActorsToDataLayer(ActorsToAdd, DataLayer);
 		}
 	}
-	else
-	{
-		if (!ActorsToAdd[0]->HasDataLayers())
-		{
-			// Only allow actors not coming from the DataLayerBrowser
-			const FScopedTransaction Transaction(LOCTEXT("AddSelectedActorsToNewDataLayer", "Add Actors to New Data Layer"));
-			if (UDataLayer* NewDataLayer = DataLayerEditorSubsystem->CreateDataLayer())
-			{
-				DataLayerEditorSubsystem->AddActorsToDataLayer(ActorsToAdd, NewDataLayer);
-			}
-		}
-	}
 }
 
 struct FWeakDataLayerActorSelector
