DataLayer: Allow default editor load state to be unloaded + Resetting user settings

diff --git a/Engine/Source/Editor/DataLayerEditor/Private/DataLayer/DataLayerEditorSubsystem.cpp b/Engine/Source/Editor/DataLayerEditor/Private/DataLayer/DataLayerEditorSubsystem.cpp
--- a/Engine/Source/Editor/DataLayerEditor/Private/DataLayer/DataLayerEditorSubsystem.cpp
+++ b/Engine/Source/Editor/DataLayerEditor/Private/DataLayer/DataLayerEditorSubsystem.cpp
@@ -187,15 +187,20 @@ void UDataLayerEditorSubsystem::UpdateDataLayerEditorPerProjectUserSettings()
 	if (AWorldDataLayers* WorldDataLayers = GetWorldDataLayers())
 	{
 		TArray<FName> DataLayersNotLoadedInEditor;
-		WorldDataLayers->ForEachDataLayer([&DataLayersNotLoadedInEditor](UDataLayer* DataLayer)
+		TArray<FName> DataLayersLoadedInEditor;
+		WorldDataLayers->ForEachDataLayer([&DataLayersNotLoadedInEditor, &DataLayersLoadedInEditor](UDataLayer* DataLayer)
 		{
-			if (!DataLayer->IsDynamicallyLoadedInEditor())
+			if (!DataLayer->IsDynamicallyLoadedInEditor() && DataLayer->IsInitiallyLoadedInEditor())
 			{
 				DataLayersNotLoadedInEditor.Add(DataLayer->GetFName());
 			}
+			else if (DataLayer->IsDynamicallyLoadedInEditor() && !DataLayer->IsInitiallyLoadedInEditor())
+			{
+				DataLayersLoadedInEditor.Add(DataLayer->GetFName());
+			}
 			return true;
 		});
-		GetMutableDefault<UWorldPartitionEditorPerProjectUserSettings>()->SetWorldDataLayersNotLoadedInEditor(GetWorld(), DataLayersNotLoadedInEditor);
+		GetMutableDefault<UWorldPartitionEditorPerProjectUserSettings>()->SetWorldDataLayersNonDefaultEditorLoadStates(GetWorld(), DataLayersLoadedInEditor, DataLayersNotLoadedInEditor);
 	}
 }
 
@@ -1068,6 +1073,16 @@ bool UDataLayerEditorSubsystem::ToggleDataLayersIsDynamicallyLoadedInEditor(cons
 	return bRefreshNeeded ? RefreshWorldPartitionEditorCells() : true;
 }
 
+bool UDataLayerEditorSubsystem::ResetUserSettings(const TArray<UDataLayer*>& DataLayers)
+{
+	bool bRefreshNeeded = false;
+	for (UDataLayer* DataLayer : DataLayers)
+	{
+		bRefreshNeeded = bRefreshNeeded || SetDataLayerIsDynamicallyLoadedInEditorInternal(DataLayer, DataLayer->IsInitiallyLoadedInEditor());
+	}
+	return bRefreshNeeded ? RefreshWorldPartitionEditorCells() : true;
+}
+
 UDataLayer* UDataLayerEditorSubsystem::GetDataLayer(const FActorDataLayer& ActorDataLayer) const
 {
 	return GetDataLayerFromName(ActorDataLayer.Name);
@@ -1165,4 +1180,4 @@ bool UDataLayerEditorSubsystem::RenameDataLayer(UDataLayer* DataLayer, const FNa
 		}
 	}
 	return false;
-}
\ No newline at end of file
+}
diff --git a/Engine/Source/Editor/DataLayerEditor/Private/DataLayer/DataLayerMode.cpp b/Engine/Source/Editor/DataLayerEditor/Private/DataLayer/DataLayerMode.cpp
--- a/Engine/Source/Editor/DataLayerEditor/Private/DataLayer/DataLayerMode.cpp
+++ b/Engine/Source/Editor/DataLayerEditor/Private/DataLayer/DataLayerMode.cpp
@@ -636,6 +636,26 @@ void FDataLayerMode::RegisterContextMenu()
 						FCanExecuteAction::CreateLambda([=] { return !AllDataLayers.IsEmpty(); })
 					));
 			}
+
+			{
+				FToolMenuSection& Section = InMenu->AddSection("DataLayerUserSettings", LOCTEXT("DataLayerUserSettings", "User Settings"));
+
+				Section.AddMenuEntry("ResetDataLayerUserSettings", LOCTEXT("ResetDataLayerUserSettings", "Reset User Settings"), FText(), FSlateIcon(),
+					FUIAction(
+						FExecuteAction::CreateLambda([=]() {
+							check(!SelectedDataLayers.IsEmpty());
+							{
+								const FScopedTransaction Transaction(LOCTEXT("ResetDataLayerUserSettings", "Reset User Settings"));
+								if (!DataLayerEditorSubsystem->ResetUserSettings(SelectedDataLayers))
+								{
+									// Nothing changed
+									GEditor->Trans->Undo();
+								}
+								
+							}}),
+						FCanExecuteAction::CreateLambda([=] { return !SelectedDataLayers.IsEmpty(); })
+								));
+			}
 		}));
 	}
 
diff --git a/Engine/Source/Editor/DataLayerEditor/Public/DataLayer/DataLayerEditorSubsystem.h b/Engine/Source/Editor/DataLayerEditor/Public/DataLayer/DataLayerEditorSubsystem.h
--- a/Engine/Source/Editor/DataLayerEditor/Public/DataLayer/DataLayerEditorSubsystem.h
+++ b/Engine/Source/Editor/DataLayerEditor/Public/DataLayer/DataLayerEditorSubsystem.h
@@ -573,6 +573,13 @@ public:
 	UFUNCTION(BlueprintCallable, Category = DataLayers)
 	bool RenameDataLayer(UDataLayer* DataLayer, const FName& NewDataLayerLabel);
 
+	/**
+	* Resets user override settings of the provided DataLayers
+	* 
+	* @param DataLayers A valid list of DataLayer.
+	*/
+	bool ResetUserSettings(const TArray<UDataLayer*>& DataLayers);
+
 private:
 
 	/**
diff --git a/Engine/Source/Runtime/Engine/Private/WorldPartition/DataLayer/DataLayer.cpp b/Engine/Source/Runtime/Engine/Private/WorldPartition/DataLayer/DataLayer.cpp
--- a/Engine/Source/Runtime/Engine/Private/WorldPartition/DataLayer/DataLayer.cpp
+++ b/Engine/Source/Runtime/Engine/Private/WorldPartition/DataLayer/DataLayer.cpp
@@ -15,6 +15,7 @@ UDataLayer::UDataLayer(const FObjectInitializer& ObjectInitializer)
 , bIsInitiallyActive_DEPRECATED(false)
 , bIsVisible(true)
 , bIsInitiallyVisible(true)
+, bIsInitiallyLoadedInEditor(true)
 , bIsDynamicallyLoadedInEditor(true)
 , bIsLocked(false)
 #endif
diff --git a/Engine/Source/Runtime/Engine/Private/WorldPartition/DataLayer/WorldDataLayers.cpp b/Engine/Source/Runtime/Engine/Private/WorldPartition/DataLayer/WorldDataLayers.cpp
--- a/Engine/Source/Runtime/Engine/Private/WorldPartition/DataLayer/WorldDataLayers.cpp
+++ b/Engine/Source/Runtime/Engine/Private/WorldPartition/DataLayer/WorldDataLayers.cpp
@@ -412,6 +412,12 @@ void AWorldDataLayers::PostLoad()
 	GetWorld()->SetWorldDataLayers(this);
 
 #if WITH_EDITOR
+	// Setup defaults before overriding with user settings
+	for (UDataLayer* DataLayer : WorldDataLayers)
+	{
+		DataLayer->SetIsDynamicallyLoadedInEditor(DataLayer->IsInitiallyLoadedInEditor());
+	}
+
 	// Initialize DataLayer's IsDynamicallyLoadedInEditor based on DataLayerEditorPerProjectUserSettings
 	const TArray<FName>& SettingsDataLayersNotLoadedInEditor = GetMutableDefault<UWorldPartitionEditorPerProjectUserSettings>()->GetWorldDataLayersNotLoadedInEditor(GetWorld());
 	for (const FName& DataLayerName : SettingsDataLayersNotLoadedInEditor)
@@ -421,6 +427,15 @@ void AWorldDataLayers::PostLoad()
 			DataLayer->SetIsDynamicallyLoadedInEditor(false);
 		}
 	}
+
+	const TArray<FName>& SettingsDataLayersLoadedInEditor = GetMutableDefault<UWorldPartitionEditorPerProjectUserSettings>()->GetWorldDataLayersLoadedInEditor(GetWorld());
+	for (const FName& DataLayerName : SettingsDataLayersLoadedInEditor)
+	{
+		if (UDataLayer* DataLayer = const_cast<UDataLayer*>(GetDataLayerFromName(DataLayerName)))
+		{
+			DataLayer->SetIsDynamicallyLoadedInEditor(true);
+}
+	}
 #else
 	// Build acceleration tables
 	for (const UDataLayer* DataLayer : WorldDataLayers)
diff --git a/Engine/Source/Runtime/Engine/Public/WorldPartition/DataLayer/DataLayer.h b/Engine/Source/Runtime/Engine/Public/WorldPartition/DataLayer/DataLayer.h
--- a/Engine/Source/Runtime/Engine/Public/WorldPartition/DataLayer/DataLayer.h
+++ b/Engine/Source/Runtime/Engine/Public/WorldPartition/DataLayer/DataLayer.h
@@ -36,6 +36,7 @@ public:
 	void SetIsLocked(bool bInIsLocked) { bIsLocked = bInIsLocked; }
 
 	bool IsDynamicallyLoadedInEditor() const { return bIsDynamicallyLoadedInEditor; }
+	bool IsInitiallyLoadedInEditor() const { return bIsInitiallyLoadedInEditor; }
 
 	static FText GetDataLayerText(const UDataLayer* InDataLayer);
 
@@ -83,7 +84,11 @@ private:
 	UPROPERTY(Category = "Data Layer - Editor", EditAnywhere)
 	uint32 bIsInitiallyVisible : 1;
 
-	/** Whether the Data Layer affects actor editor loading */
+	/** Determines the default value of the data layer's loaded state in editor if it hasn't been changed in data layer outliner by the user */
+	UPROPERTY(Category = "Data Layer - Editor", EditAnywhere, meta = (DisplayName = "Is Initially Loaded"))
+	uint32 bIsInitiallyLoadedInEditor : 1;
+
+	/** Wheter the data layer is loaded in editor (user setting) */
 	UPROPERTY(Transient)
 	uint32 bIsDynamicallyLoadedInEditor : 1;
 
diff --git a/Engine/Source/Runtime/Engine/Public/WorldPartition/WorldPartitionEditorPerProjectUserSettings.h b/Engine/Source/Runtime/Engine/Public/WorldPartition/WorldPartitionEditorPerProjectUserSettings.h
--- a/Engine/Source/Runtime/Engine/Public/WorldPartition/WorldPartitionEditorPerProjectUserSettings.h
+++ b/Engine/Source/Runtime/Engine/Public/WorldPartition/WorldPartitionEditorPerProjectUserSettings.h
@@ -29,6 +29,9 @@ struct FWorldPartitionPerWorldSettings
 
 	UPROPERTY()
 	TArray<FName> NotLoadedDataLayers;
+
+	UPROPERTY()
+	TArray<FName> LoadedDataLayers;
 #endif
 };
 
@@ -97,12 +100,19 @@ public:
 		return PerWorldEditorSettings.FindOrAdd(TSoftObjectPtr<UWorld>(InWorld)).NotLoadedDataLayers;
 	}
 
-	void SetWorldDataLayersNotLoadedInEditor(UWorld* InWorld, const TArray<FName>& InDataLayersLoadedInEditor)
+	const TArray<FName>& GetWorldDataLayersLoadedInEditor(UWorld* InWorld)
+	{
+		return PerWorldEditorSettings.FindOrAdd(TSoftObjectPtr<UWorld>(InWorld)).LoadedDataLayers;
+	}
+
+	void SetWorldDataLayersNonDefaultEditorLoadStates(UWorld* InWorld, const TArray<FName>& InDataLayersLoadedInEditor, const TArray<FName>& InDataLayersNotLoadedInEditor)
 	{
 		if (ShouldSaveSettings(InWorld))
 		{
-			TArray<FName>& DataLayersNotLoadedInEditor = PerWorldEditorSettings.FindOrAdd(TSoftObjectPtr<UWorld>(InWorld)).NotLoadedDataLayers;
-			DataLayersNotLoadedInEditor = InDataLayersLoadedInEditor;
+			FWorldPartitionPerWorldSettings& PerWorldSettings = PerWorldEditorSettings.FindOrAdd(TSoftObjectPtr<UWorld>(InWorld));
+
+			PerWorldSettings.NotLoadedDataLayers = InDataLayersNotLoadedInEditor;
+			PerWorldSettings.LoadedDataLayers = InDataLayersLoadedInEditor;
 			SaveConfig();
 		}
 	}
