DataLayer: fix badly formed condition resulting in missing calls. (was breaking multi-select toggle in Data Layer UI)

diff --git a/Engine/Source/Editor/DataLayerEditor/Private/DataLayer/DataLayerEditorSubsystem.cpp b/Engine/Source/Editor/DataLayerEditor/Private/DataLayer/DataLayerEditorSubsystem.cpp
--- a/Engine/Source/Editor/DataLayerEditor/Private/DataLayer/DataLayerEditorSubsystem.cpp
+++ b/Engine/Source/Editor/DataLayerEditor/Private/DataLayer/DataLayerEditorSubsystem.cpp
@@ -1007,7 +1007,10 @@ bool UDataLayerEditorSubsystem::SetDataLayersIsDynamicallyLoaded(const TArray<UD
 	bool bRefreshNeeded = false;
 	for (UDataLayer* DataLayer : DataLayers)
 	{
-		bRefreshNeeded = bRefreshNeeded || SetDataLayerIsDynamicallyLoadedInternal(DataLayer, bIsDynamicallyLoaded);
+		if (SetDataLayerIsDynamicallyLoadedInternal(DataLayer, bIsDynamicallyLoaded))
+		{
+			bRefreshNeeded = true;
+		}
 	}
 	return bRefreshNeeded ? RefreshWorldPartitionEditorCells() : true;
 }
@@ -1023,7 +1026,10 @@ bool UDataLayerEditorSubsystem::ToggleDataLayersIsDynamicallyLoaded(const TArray
 	bool bRefreshNeeded = false;
 	for (UDataLayer* DataLayer : DataLayers)
 	{
-		bRefreshNeeded = bRefreshNeeded || SetDataLayerIsDynamicallyLoadedInternal(DataLayer, !DataLayer->IsDynamicallyLoaded());
+		if (SetDataLayerIsDynamicallyLoadedInternal(DataLayer, !DataLayer->IsDynamicallyLoaded()))
+		{
+			bRefreshNeeded = true;
+		}
 	}
 	return bRefreshNeeded ? RefreshWorldPartitionEditorCells() : true;
 }
@@ -1052,7 +1058,10 @@ bool UDataLayerEditorSubsystem::SetDataLayersIsDynamicallyLoadedInEditor(const T
 	bool bRefreshNeeded = false;
 	for (UDataLayer* DataLayer : DataLayers)
 	{
-		bRefreshNeeded = bRefreshNeeded || SetDataLayerIsDynamicallyLoadedInEditorInternal(DataLayer, bIsDynamicallyLoadedInEditor);
+		if (SetDataLayerIsDynamicallyLoadedInEditorInternal(DataLayer, bIsDynamicallyLoadedInEditor))
+		{
+			bRefreshNeeded = true;
+		}
 	}
 	return bRefreshNeeded ? RefreshWorldPartitionEditorCells() : true;
 }
@@ -1068,7 +1077,10 @@ bool UDataLayerEditorSubsystem::ToggleDataLayersIsDynamicallyLoadedInEditor(cons
 	bool bRefreshNeeded = false;
 	for (UDataLayer* DataLayer : DataLayers)
 	{
-		bRefreshNeeded = bRefreshNeeded || SetDataLayerIsDynamicallyLoadedInEditorInternal(DataLayer, !DataLayer->IsDynamicallyLoadedInEditor());
+		if (SetDataLayerIsDynamicallyLoadedInEditorInternal(DataLayer, !DataLayer->IsDynamicallyLoadedInEditor()))
+		{
+			bRefreshNeeded = true;
+		}
 	}
 	return bRefreshNeeded ? RefreshWorldPartitionEditorCells() : true;
 }
@@ -1078,7 +1090,10 @@ bool UDataLayerEditorSubsystem::ResetUserSettings(const TArray<UDataLayer*>& Dat
 	bool bRefreshNeeded = false;
 	for (UDataLayer* DataLayer : DataLayers)
 	{
-		bRefreshNeeded = bRefreshNeeded || SetDataLayerIsDynamicallyLoadedInEditorInternal(DataLayer, DataLayer->IsInitiallyLoadedInEditor());
+		if (SetDataLayerIsDynamicallyLoadedInEditorInternal(DataLayer, DataLayer->IsInitiallyLoadedInEditor()))
+		{
+			bRefreshNeeded = true;
+		}
 	}
 	return bRefreshNeeded ? RefreshWorldPartitionEditorCells() : true;
 }
