Only close self when choosing data layers
  - This prevents a crash when selecting a data layer on a sequencer track.

diff --git a/Engine/Source/Editor/DataLayerEditor/Private/DataLayer/DataLayerPropertyTypeCustomizationHelper.cpp b/Engine/Source/Editor/DataLayerEditor/Private/DataLayer/DataLayerPropertyTypeCustomizationHelper.cpp
--- a/Engine/Source/Editor/DataLayerEditor/Private/DataLayer/DataLayerPropertyTypeCustomizationHelper.cpp
+++ b/Engine/Source/Editor/DataLayerEditor/Private/DataLayer/DataLayerPropertyTypeCustomizationHelper.cpp
@@ -11,7 +11,7 @@
 
 TSharedRef<SWidget> FDataLayerPropertyTypeCustomizationHelper::CreateDataLayerMenu(TFunction<void(const UDataLayer* DataLayer)> OnDataLayerSelectedFunction)
 {
-	FMenuBuilder MenuBuilder(true, nullptr);
+	FMenuBuilder MenuBuilder(true, nullptr, nullptr, true);
 
 	MenuBuilder.AddMenuEntry(
 		LOCTEXT("OpenDataLayersBrowser", "Browse Data Layers..."),
