Fixed LevelInstance actors not properly reacting to DataLayer visibility change in the editor.

diff --git a/Engine/Source/Editor/DataLayerEditor/Private/DataLayer/DataLayerEditorSubsystem.cpp b/Engine/Source/Editor/DataLayerEditor/Private/DataLayer/DataLayerEditorSubsystem.cpp
--- a/Engine/Source/Editor/DataLayerEditor/Private/DataLayer/DataLayerEditorSubsystem.cpp
+++ b/Engine/Source/Editor/DataLayerEditor/Private/DataLayer/DataLayerEditorSubsystem.cpp
@@ -234,7 +234,7 @@ void UDataLayerEditorSubsystem::PostUndoRedo()
 
 bool UDataLayerEditorSubsystem::IsActorValidForDataLayer(AActor* Actor)
 {
-	return Actor && Actor->IsValidForDataLayer();
+	return Actor && Actor->IsValidForDataLayer() && (Actor->GetLevel() == Actor->GetWorld()->PersistentLevel);
 }
 
 void UDataLayerEditorSubsystem::InitializeNewActorDataLayers(AActor* Actor)
@@ -684,37 +684,6 @@ void UDataLayerEditorSubsystem::RemoveViewFromActorViewVisibility(FLevelEditorVi
 	}
 }
 
-static void UpdateBrushDataLayerVisibility(ABrush* Brush, bool bIsHidden)
-{
-	ULevel* Level = Brush->GetLevel();
-	if (!Level)
-	{
-		return;
-	}
-
-	UModel* Model = Level->Model;
-	if (!Model)
-	{
-		return;
-	}
-
-	bool bAnySurfaceWasFound = false;
-	for (FBspSurf& Surf : Model->Surfs)
-	{
-		if (Surf.Actor == Brush)
-		{
-			Surf.bHiddenEdLayer = bIsHidden;
-			bAnySurfaceWasFound = true;
-		}
-	}
-
-	if (bAnySurfaceWasFound)
-	{
-		Level->UpdateModelComponents();
-		Model->InvalidSurfaces = true;
-	}
-}
-
 bool UDataLayerEditorSubsystem::UpdateActorVisibility(AActor* Actor, bool& bOutSelectionChanged, bool& bOutActorModified, const bool bNotifySelectionChange, const bool bRedrawViewports)
 {
 	bOutActorModified = false;
@@ -728,15 +697,8 @@ bool UDataLayerEditorSubsystem::UpdateActorVisibility(AActor* Actor, bool& bOutS
 	// If the actor doesn't belong to any DataLayers
 	if (!Actor->HasValidDataLayers())
 	{
-		// If the actor is also hidden
-		if (Actor->bHiddenEdLayer)
-		{
-			// Actors that don't belong to any DataLayer shouldn't be hidden
-			Actor->bHiddenEdLayer = false;
-			Actor->MarkComponentsRenderStateDirty();
-			bOutActorModified = true;
-		}
-
+		// Actors that don't belong to any DataLayer shouldn't be hidden
+		bOutActorModified = Actor->SetIsHiddenEdLayer(false);
 		return bOutActorModified;
 	}
 
@@ -747,18 +709,10 @@ bool UDataLayerEditorSubsystem::UpdateActorVisibility(AActor* Actor, bool& bOutS
 		{
 			if (DataLayer->IsVisible() && Actor->ContainsDataLayer(DataLayer))
 			{
-				if (Actor->bHiddenEdLayer)
+				if (Actor->SetIsHiddenEdLayer(false))
 				{
-					Actor->bHiddenEdLayer = false;
-					Actor->MarkComponentsRenderStateDirty();
 					bOutActorModified = true;
-
-					if (ABrush* Brush = Cast<ABrush>(Actor))
-					{
-						UpdateBrushDataLayerVisibility(Brush, /*bIsHidden*/false);
-					}
 				}
-
 				// Stop, because we found at least one visible DataLayer the actor belongs to
 				bActorBelongsToVisibleDataLayer = true;
 				return false;
@@ -770,18 +724,11 @@ bool UDataLayerEditorSubsystem::UpdateActorVisibility(AActor* Actor, bool& bOutS
 	// If the actor isn't part of a visible DataLayer, hide and de-select it.
 	if (!bActorBelongsToVisibleDataLayer)
 	{
-		if (!Actor->bHiddenEdLayer)
+		if (Actor->SetIsHiddenEdLayer(true))
 		{
-			Actor->bHiddenEdLayer = true;
-			Actor->MarkComponentsRenderStateDirty();
 			bOutActorModified = true;
-
-			if (ABrush* Brush = Cast<ABrush>(Actor))
-			{
-				UpdateBrushDataLayerVisibility(Brush, /*bIsHidden*/false);
-			}
 		}
-
+		
 		// If the actor was selected, mark it as unselected
 		if (Actor->IsSelected())
 		{
diff --git a/Engine/Source/Runtime/Engine/Classes/Engine/Brush.h b/Engine/Source/Runtime/Engine/Classes/Engine/Brush.h
--- a/Engine/Source/Runtime/Engine/Classes/Engine/Brush.h
+++ b/Engine/Source/Runtime/Engine/Classes/Engine/Brush.h
@@ -169,7 +169,7 @@ public:
 	virtual void PostRegisterAllComponents() override;
 	virtual void CheckForErrors() override;
 	virtual void SetIsTemporarilyHiddenInEditor( bool bIsHidden ) override;
-
+	virtual bool SetIsHiddenEdLayer(bool bIsHiddenEdLayer) override;
 public:
 
 	virtual void InitPosRotScale();
diff --git a/Engine/Source/Runtime/Engine/Classes/GameFramework/Actor.h b/Engine/Source/Runtime/Engine/Classes/GameFramework/Actor.h
--- a/Engine/Source/Runtime/Engine/Classes/GameFramework/Actor.h
+++ b/Engine/Source/Runtime/Engine/Classes/GameFramework/Actor.h
@@ -2144,6 +2144,9 @@ public:
 	UFUNCTION(BlueprintCallable, Category="Editor Scripting | Actor Editing")
 	virtual void SetIsTemporarilyHiddenInEditor( bool bIsHidden );
 
+	/** Changes bHiddenEdLayer flag and returns true if flag changed. */
+	virtual bool SetIsHiddenEdLayer(bool bIsHiddenEdLayer);
+
 	/** Returns if level should keep a reference to the external actor for PIE (used for always loaded actors). */
 	bool IsForceExternalActorLevelReferenceForPIE() const
 	{
diff --git a/Engine/Source/Runtime/Engine/Private/ActorEditor.cpp b/Engine/Source/Runtime/Engine/Private/ActorEditor.cpp
--- a/Engine/Source/Runtime/Engine/Private/ActorEditor.cpp
+++ b/Engine/Source/Runtime/Engine/Private/ActorEditor.cpp
@@ -727,6 +727,17 @@ void AActor::SetIsTemporarilyHiddenInEditor( bool bIsHidden )
 	}
 }
 
+bool AActor::SetIsHiddenEdLayer(bool bIsHiddenEdLayer)
+{
+	if (bHiddenEdLayer != bIsHiddenEdLayer)
+	{
+		bHiddenEdLayer = bIsHiddenEdLayer;
+		MarkComponentsRenderStateDirty();
+		return true;
+	}
+	return false;
+}
+
 bool AActor::IsEditable() const
 {
 	return bEditable;
diff --git a/Engine/Source/Runtime/Engine/Private/Brush.cpp b/Engine/Source/Runtime/Engine/Private/Brush.cpp
--- a/Engine/Source/Runtime/Engine/Private/Brush.cpp
+++ b/Engine/Source/Runtime/Engine/Private/Brush.cpp
@@ -184,6 +184,35 @@ void ABrush::SetIsTemporarilyHiddenInEditor( bool bIsHidden )
 	}
 }
 
+bool ABrush::SetIsHiddenEdLayer(bool bIsHiddenEdLayer)
+{
+	if (Super::SetIsHiddenEdLayer(bIsHiddenEdLayer))
+	{
+		ULevel* Level = GetLevel();
+		UModel* Model = Level ? ToRawPtr(Level->Model) : nullptr;
+		if (Level && Model)
+		{
+			bool bAnySurfaceWasFound = false;
+			for (FBspSurf& Surf : Model->Surfs)
+			{
+				if (Surf.Actor == this)
+				{
+					Surf.bHiddenEdLayer = bIsHiddenEdLayer;
+					bAnySurfaceWasFound = true;
+				}
+			}
+
+			if (bAnySurfaceWasFound)
+			{
+				Level->UpdateModelComponents();
+				Model->InvalidSurfaces = true;
+			}
+		}
+		return true;
+	}
+	return false;
+}
+
 void ABrush::PostLoad()
 {
 	Super::PostLoad();
diff --git a/Engine/Source/Runtime/Engine/Private/LevelInstance/LevelInstanceActor.cpp b/Engine/Source/Runtime/Engine/Private/LevelInstance/LevelInstanceActor.cpp
--- a/Engine/Source/Runtime/Engine/Private/LevelInstance/LevelInstanceActor.cpp
+++ b/Engine/Source/Runtime/Engine/Private/LevelInstance/LevelInstanceActor.cpp
@@ -550,6 +550,23 @@ void ALevelInstance::SetIsTemporarilyHiddenInEditor(bool bIsHidden)
 	}
 }
 
+bool ALevelInstance::SetIsHiddenEdLayer(bool bIsHiddenEdLayer)
+{
+	bool bHasChanged = false;
+	{
+		TGuardValue<bool> LoadUnloadGuard(bGuardLoadUnload, true);
+		bHasChanged = Super::SetIsHiddenEdLayer(bIsHiddenEdLayer);
+	}
+	if (bHasChanged)
+	{
+		if (ULevelInstanceSubsystem* LevelInstanceSubsystem = GetLevelInstanceSubsystem())
+		{
+			LevelInstanceSubsystem->SetIsHiddenEdLayer(this, bIsHiddenEdLayer);
+		}
+	}
+	return bHasChanged;
+}
+
 void ALevelInstance::EditorGetUnderlyingActors(TSet<AActor*>& OutUnderlyingActors) const
 {
 	Super::EditorGetUnderlyingActors(OutUnderlyingActors);
diff --git a/Engine/Source/Runtime/Engine/Private/LevelInstance/LevelInstanceSubsystem.cpp b/Engine/Source/Runtime/Engine/Private/LevelInstance/LevelInstanceSubsystem.cpp
--- a/Engine/Source/Runtime/Engine/Private/LevelInstance/LevelInstanceSubsystem.cpp
+++ b/Engine/Source/Runtime/Engine/Private/LevelInstance/LevelInstanceSubsystem.cpp
@@ -636,6 +636,18 @@ bool ULevelInstanceSubsystem::HasEditingChildrenLevelInstances(const ALevelInsta
 	return bEditingChildren;
 }
 
+void ULevelInstanceSubsystem::SetIsHiddenEdLayer(ALevelInstance* LevelInstanceActor, bool bIsHiddenEdLayer)
+{
+	if (ULevel* LevelInstanceLevel = GetLevelInstanceLevel(LevelInstanceActor))
+	{
+		ForEachActorInLevel(LevelInstanceLevel, [bIsHiddenEdLayer](AActor* LevelActor)
+		{
+			LevelActor->SetIsHiddenEdLayer(bIsHiddenEdLayer);
+			return true;
+		});
+	}
+}
+
 void ULevelInstanceSubsystem::SetIsTemporarilyHiddenInEditor(ALevelInstance* LevelInstanceActor, bool bIsHidden)
 {
 	if (ULevel* LevelInstanceLevel = GetLevelInstanceLevel(LevelInstanceActor))
diff --git a/Engine/Source/Runtime/Engine/Public/LevelInstance/LevelInstanceActor.h b/Engine/Source/Runtime/Engine/Public/LevelInstance/LevelInstanceActor.h
--- a/Engine/Source/Runtime/Engine/Public/LevelInstance/LevelInstanceActor.h
+++ b/Engine/Source/Runtime/Engine/Public/LevelInstance/LevelInstanceActor.h
@@ -68,6 +68,7 @@ public:
 	virtual TUniquePtr<class FWorldPartitionActorDesc> CreateClassActorDesc() const override;
 	virtual bool CanDeleteSelectedActor(FText& OutReason) const override;
 	virtual void SetIsTemporarilyHiddenInEditor(bool bIsHidden) override;
+	virtual bool SetIsHiddenEdLayer(bool bIsHiddenEdLayer) override;
 	virtual void EditorGetUnderlyingActors(TSet<AActor*>& OutUnderlyingActors) const override;
 	virtual void PushSelectionToProxies() override;
 	virtual void PushLevelInstanceEditingStateToProxies(bool bInEditingState) override;
diff --git a/Engine/Source/Runtime/Engine/Public/LevelInstance/LevelInstanceSubsystem.h b/Engine/Source/Runtime/Engine/Public/LevelInstance/LevelInstanceSubsystem.h
--- a/Engine/Source/Runtime/Engine/Public/LevelInstance/LevelInstanceSubsystem.h
+++ b/Engine/Source/Runtime/Engine/Public/LevelInstance/LevelInstanceSubsystem.h
@@ -70,6 +70,7 @@ public:
 	bool HasDirtyChildrenLevelInstances(const ALevelInstance* LevelInstanceActor) const;
 	bool HasEditingChildrenLevelInstances(const ALevelInstance* LevelInstanceActor) const;
 
+	void SetIsHiddenEdLayer(ALevelInstance* LevelInstanceActor, bool bIsHiddenEdLayer);
 	void SetIsTemporarilyHiddenInEditor(ALevelInstance* LevelInstanceActor, bool bIsHidden);
 
 	bool SetCurrent(ALevelInstance* LevelInstanceActor) const;
