Added BP functions in UDataLayerEditorSubsystem to clear actor's DataLayers.

diff --git a/Engine/Source/Editor/DataLayerEditor/Private/DataLayer/DataLayerEditorSubsystem.cpp b/Engine/Source/Editor/DataLayerEditor/Private/DataLayer/DataLayerEditorSubsystem.cpp
--- a/Engine/Source/Editor/DataLayerEditor/Private/DataLayer/DataLayerEditorSubsystem.cpp
+++ b/Engine/Source/Editor/DataLayerEditor/Private/DataLayer/DataLayerEditorSubsystem.cpp
@@ -326,6 +326,44 @@ bool UDataLayerEditorSubsystem::AddActorsToDataLayers(const TArray<AActor*>& Act
 	return bChangesOccurred;
 }
 
+bool UDataLayerEditorSubsystem::RemoveActorFromAllDataLayers(AActor* Actor)
+{
+	return RemoveActorsFromAllDataLayers({ Actor });
+}
+
+bool UDataLayerEditorSubsystem::RemoveActorsFromAllDataLayers(const TArray<AActor*>& Actors)
+{
+	GEditor->GetSelectedActors()->BeginBatchSelectOperation();
+
+	bool bChangesOccurred = false;
+	for (AActor* Actor : Actors)
+	{
+		TArray<const UDataLayer*> ModifiedDataLayers = Actor->GetDataLayerObjects();
+		if (Actor->RemoveAllDataLayers())
+		{
+			for (const UDataLayer* DataLayer : ModifiedDataLayers)
+			{
+				DataLayerChanged.Broadcast(EDataLayerAction::Modify, DataLayer, NAME_None);
+			}
+			ActorDataLayersChanged.Broadcast(Actor);
+
+			// Update per-view visibility info
+			UpdateActorAllViewsVisibility(Actor);
+
+			// Update general actor visibility
+			bool bActorModified = false;
+			bool bActorSelectionChanged = false;
+			UpdateActorVisibility(Actor, bActorSelectionChanged, bActorModified, /*bActorNotifySelectionChange*/true, /*bActorRedrawViewports*/false);
+
+			bChangesOccurred = true;
+		}
+	}
+
+	GEditor->GetSelectedActors()->EndBatchSelectOperation();
+
+	return bChangesOccurred;
+}
+
 bool UDataLayerEditorSubsystem::RemoveActorFromDataLayer(AActor* Actor, UDataLayer* DataLayer)
 {
 	TArray<AActor*> Actors;
diff --git a/Engine/Source/Editor/DataLayerEditor/Public/DataLayer/DataLayerEditorSubsystem.h b/Engine/Source/Editor/DataLayerEditor/Public/DataLayer/DataLayerEditorSubsystem.h
--- a/Engine/Source/Editor/DataLayerEditor/Public/DataLayer/DataLayerEditorSubsystem.h
+++ b/Engine/Source/Editor/DataLayerEditor/Public/DataLayer/DataLayerEditorSubsystem.h
@@ -110,6 +110,23 @@ public:
 	UFUNCTION(BlueprintCallable, Category = DataLayers)
 	bool RemoveActorFromDataLayers(AActor* Actor, const TArray<UDataLayer*>& DataLayers);
 
+	/**
+	 * Removes an actor from all DataLayers.
+	 *
+	 * @param	Actor			The actor to modify
+	 * @return					true if the actor was changed.
+	 */
+	UFUNCTION(BlueprintCallable, Category = DataLayers)
+	bool RemoveActorFromAllDataLayers(AActor* Actor);
+	
+	/**
+	 * Removes an actor from all DataLayers.
+	 *
+	 * @param	Actor			The actors to modify
+	 * @return					true if any actor was changed.
+	 */
+	UFUNCTION(BlueprintCallable, Category = DataLayers)
+	bool RemoveActorsFromAllDataLayers(const TArray<AActor*>& Actors);
 
 	/////////////////////////////////////////////////
 	// Operations on multiple actors.
diff --git a/Engine/Source/Runtime/Engine/Classes/GameFramework/Actor.h b/Engine/Source/Runtime/Engine/Classes/GameFramework/Actor.h
--- a/Engine/Source/Runtime/Engine/Classes/GameFramework/Actor.h
+++ b/Engine/Source/Runtime/Engine/Classes/GameFramework/Actor.h
@@ -1090,6 +1090,7 @@ public:
 	// DataLayers functions.
 	bool AddDataLayer(const UDataLayer* DataLayer);
 	bool RemoveDataLayer(const UDataLayer* DataLayer);
+	bool RemoveAllDataLayers();
 	bool ContainsDataLayer(const UDataLayer* DataLayer) const;
 	virtual bool SupportsDataLayer() const { return true; }
 	bool HasDataLayers() const;
diff --git a/Engine/Source/Runtime/Engine/Private/ActorEditor.cpp b/Engine/Source/Runtime/Engine/Private/ActorEditor.cpp
--- a/Engine/Source/Runtime/Engine/Private/ActorEditor.cpp
+++ b/Engine/Source/Runtime/Engine/Private/ActorEditor.cpp
@@ -1124,6 +1124,17 @@ bool AActor::RemoveDataLayer(const UDataLayer* DataLayer)
 	return bActorWasModified;
 }
 
+bool AActor::RemoveAllDataLayers()
+{
+	if (HasDataLayers())
+	{
+		Modify();
+		DataLayers.Empty();
+		return true;
+	}
+	return false;
+}
+
 bool AActor::ContainsDataLayer(const UDataLayer* DataLayer) const
 {
 	return DataLayer && DataLayers.Contains(FActorDataLayer(DataLayer->GetFName()));
