Only include changes made by user operations to the per-user settings for data layers.

diff --git a/Engine/Source/Editor/DataLayerEditor/Private/DataLayer/DataLayerEditorSubsystem.cpp b/Engine/Source/Editor/DataLayerEditor/Private/DataLayer/DataLayerEditorSubsystem.cpp
--- a/Engine/Source/Editor/DataLayerEditor/Private/DataLayer/DataLayerEditorSubsystem.cpp
+++ b/Engine/Source/Editor/DataLayerEditor/Private/DataLayer/DataLayerEditorSubsystem.cpp
@@ -186,17 +186,36 @@ void UDataLayerEditorSubsystem::UpdateDataLayerEditorPerProjectUserSettings()
 {
 	if (AWorldDataLayers* WorldDataLayers = GetWorldDataLayers())
 	{
+		const TArray<FName>& SettingsDataLayersNotLoadedInEditor = GetMutableDefault<UWorldPartitionEditorPerProjectUserSettings>()->GetWorldDataLayersNotLoadedInEditor(GetWorld());
+		const TArray<FName>& SettingsDataLayersLoadedInEditor = GetMutableDefault<UWorldPartitionEditorPerProjectUserSettings>()->GetWorldDataLayersLoadedInEditor(GetWorld());
+
 		TArray<FName> DataLayersNotLoadedInEditor;
 		TArray<FName> DataLayersLoadedInEditor;
-		WorldDataLayers->ForEachDataLayer([&DataLayersNotLoadedInEditor, &DataLayersLoadedInEditor](UDataLayer* DataLayer)
+		WorldDataLayers->ForEachDataLayer([&DataLayersNotLoadedInEditor, &DataLayersLoadedInEditor, &SettingsDataLayersNotLoadedInEditor, &SettingsDataLayersLoadedInEditor](UDataLayer* DataLayer)
 		{
-			if (!DataLayer->IsDynamicallyLoadedInEditor() && DataLayer->IsInitiallyLoadedInEditor())
-			{
-				DataLayersNotLoadedInEditor.Add(DataLayer->GetFName());
+			if (DataLayer->IsDynamicallyLoadedInEditorChangedByUserOperation())
+			{			
+				if (!DataLayer->IsDynamicallyLoadedInEditor() && DataLayer->IsInitiallyLoadedInEditor())
+				{
+					DataLayersNotLoadedInEditor.Add(DataLayer->GetFName());
+				}
+				else if (DataLayer->IsDynamicallyLoadedInEditor() && !DataLayer->IsInitiallyLoadedInEditor())
+				{
+					DataLayersLoadedInEditor.Add(DataLayer->GetFName());
+				}
+
+				DataLayer->ClearDynamicallyLoadedInEditorChangedByUserOperation();
 			}
-			else if (DataLayer->IsDynamicallyLoadedInEditor() && !DataLayer->IsInitiallyLoadedInEditor())
+			else
 			{
-				DataLayersLoadedInEditor.Add(DataLayer->GetFName());
+				if (SettingsDataLayersNotLoadedInEditor.Contains(DataLayer->GetFName()))
+				{
+					DataLayersNotLoadedInEditor.Add(DataLayer->GetFName());
+				}
+				else if (SettingsDataLayersLoadedInEditor.Contains(DataLayer->GetFName()))
+				{
+					DataLayersLoadedInEditor.Add(DataLayer->GetFName());
+				}
 			}
 			return true;
 		});
@@ -992,10 +1011,7 @@ bool UDataLayerEditorSubsystem::SetDataLayersIsDynamicallyLoaded(const TArray<UD
 	bool bRefreshNeeded = false;
 	for (UDataLayer* DataLayer : DataLayers)
 	{
-		if (SetDataLayerIsDynamicallyLoadedInternal(DataLayer, bIsDynamicallyLoaded))
-		{
-			bRefreshNeeded = true;
-		}
+		bRefreshNeeded |= SetDataLayerIsDynamicallyLoadedInternal(DataLayer, bIsDynamicallyLoaded);
 	}
 	return bRefreshNeeded ? RefreshWorldPartitionEditorCells() : true;
 }
@@ -1011,61 +1027,52 @@ bool UDataLayerEditorSubsystem::ToggleDataLayersIsDynamicallyLoaded(const TArray
 	bool bRefreshNeeded = false;
 	for (UDataLayer* DataLayer : DataLayers)
 	{
-		if (SetDataLayerIsDynamicallyLoadedInternal(DataLayer, !DataLayer->IsDynamicallyLoaded()))
-		{
-			bRefreshNeeded = true;
-		}
+		bRefreshNeeded |= SetDataLayerIsDynamicallyLoadedInternal(DataLayer, !DataLayer->IsDynamicallyLoaded());
 	}
 	return bRefreshNeeded ? RefreshWorldPartitionEditorCells() : true;
 }
 
-bool UDataLayerEditorSubsystem::SetDataLayerIsDynamicallyLoadedInEditorInternal(UDataLayer* DataLayer, const bool bIsDynamicallyLoadedInEditor)
+bool UDataLayerEditorSubsystem::SetDataLayerIsDynamicallyLoadedInEditorInternal(UDataLayer* DataLayer, const bool bIsDynamicallyLoadedInEditor, const bool bIsFromUserChange)
 {
 	check(DataLayer);
 	if (DataLayer->IsDynamicallyLoadedInEditor() != bIsDynamicallyLoadedInEditor)
 	{
 		DataLayer->Modify(false);
-		DataLayer->SetIsDynamicallyLoadedInEditor(bIsDynamicallyLoadedInEditor);
+		DataLayer->SetIsDynamicallyLoadedInEditor(bIsDynamicallyLoadedInEditor, bIsFromUserChange);
 		DataLayerChanged.Broadcast(EDataLayerAction::Modify, DataLayer, "bIsDynamicallyLoadedInEditor");
 		return true;
 	}
 	return false;
 }
 
-bool UDataLayerEditorSubsystem::SetDataLayerIsDynamicallyLoadedInEditor(UDataLayer* DataLayer, const bool bIsDynamicallyLoadedInEditor)
+bool UDataLayerEditorSubsystem::SetDataLayerIsDynamicallyLoadedInEditor(UDataLayer* DataLayer, const bool bIsDynamicallyLoadedInEditor, const bool bIsFromUserChange)
 {
-	bool bRefreshNeeded = SetDataLayerIsDynamicallyLoadedInEditorInternal(DataLayer, bIsDynamicallyLoadedInEditor);
+	bool bRefreshNeeded = SetDataLayerIsDynamicallyLoadedInEditorInternal(DataLayer, bIsDynamicallyLoadedInEditor, bIsFromUserChange);
 	return bRefreshNeeded ? RefreshWorldPartitionEditorCells() : true;
 }
 
-bool UDataLayerEditorSubsystem::SetDataLayersIsDynamicallyLoadedInEditor(const TArray<UDataLayer*>& DataLayers, const bool bIsDynamicallyLoadedInEditor)
+bool UDataLayerEditorSubsystem::SetDataLayersIsDynamicallyLoadedInEditor(const TArray<UDataLayer*>& DataLayers, const bool bIsDynamicallyLoadedInEditor, const bool bIsFromUserChange)
 {
 	bool bRefreshNeeded = false;
 	for (UDataLayer* DataLayer : DataLayers)
 	{
-		if (SetDataLayerIsDynamicallyLoadedInEditorInternal(DataLayer, bIsDynamicallyLoadedInEditor))
-		{
-			bRefreshNeeded = true;
-		}
+		bRefreshNeeded |= SetDataLayerIsDynamicallyLoadedInEditorInternal(DataLayer, bIsDynamicallyLoadedInEditor, bIsFromUserChange);
 	}
 	return bRefreshNeeded ? RefreshWorldPartitionEditorCells() : true;
 }
 
-bool UDataLayerEditorSubsystem::ToggleDataLayerIsDynamicallyLoadedInEditor(UDataLayer* DataLayer)
+bool UDataLayerEditorSubsystem::ToggleDataLayerIsDynamicallyLoadedInEditor(UDataLayer* DataLayer, const bool bIsFromUserChange)
 {
 	check(DataLayer);
-	return SetDataLayerIsDynamicallyLoadedInEditor(DataLayer, !DataLayer->IsDynamicallyLoadedInEditor());
+	return SetDataLayerIsDynamicallyLoadedInEditor(DataLayer, !DataLayer->IsDynamicallyLoadedInEditor(), bIsFromUserChange);
 }
 
-bool UDataLayerEditorSubsystem::ToggleDataLayersIsDynamicallyLoadedInEditor(const TArray<UDataLayer*>& DataLayers)
+bool UDataLayerEditorSubsystem::ToggleDataLayersIsDynamicallyLoadedInEditor(const TArray<UDataLayer*>& DataLayers, const bool bIsFromUserChange)
 {
 	bool bRefreshNeeded = false;
 	for (UDataLayer* DataLayer : DataLayers)
 	{
-		if (SetDataLayerIsDynamicallyLoadedInEditorInternal(DataLayer, !DataLayer->IsDynamicallyLoadedInEditor()))
-		{
-			bRefreshNeeded = true;
-		}
+		bRefreshNeeded |= SetDataLayerIsDynamicallyLoadedInEditorInternal(DataLayer, !DataLayer->IsDynamicallyLoadedInEditor(), bIsFromUserChange);
 	}
 	return bRefreshNeeded ? RefreshWorldPartitionEditorCells() : true;
 }
@@ -1075,10 +1082,7 @@ bool UDataLayerEditorSubsystem::ResetUserSettings(const TArray<UDataLayer*>& Dat
 	bool bRefreshNeeded = false;
 	for (UDataLayer* DataLayer : DataLayers)
 	{
-		if (SetDataLayerIsDynamicallyLoadedInEditorInternal(DataLayer, DataLayer->IsInitiallyLoadedInEditor()))
-		{
-			bRefreshNeeded = true;
-		}
+		bRefreshNeeded |= SetDataLayerIsDynamicallyLoadedInEditorInternal(DataLayer, DataLayer->IsInitiallyLoadedInEditor(), true);
 	}
 	return bRefreshNeeded ? RefreshWorldPartitionEditorCells() : true;
 }
diff --git a/Engine/Source/Editor/DataLayerEditor/Private/DataLayer/DataLayerOutlinerIsDynamicallyLoadedColumn.cpp b/Engine/Source/Editor/DataLayerEditor/Private/DataLayer/DataLayerOutlinerIsDynamicallyLoadedColumn.cpp
--- a/Engine/Source/Editor/DataLayerEditor/Private/DataLayer/DataLayerOutlinerIsDynamicallyLoadedColumn.cpp
+++ b/Engine/Source/Editor/DataLayerEditor/Private/DataLayer/DataLayerOutlinerIsDynamicallyLoadedColumn.cpp
@@ -97,12 +97,12 @@ const TSharedRef<SWidget> FDataLayerOutlinerIsDynamicallyLoadedColumn::Construct
 							}
 
 							const FScopedTransaction Transaction(LOCTEXT("ToggleDataLayersIsDynamicallyLoadedInEditor", "Toggle Data Layers Dynamically Loaded In Editor Flag"));
-							bSuccess = DataLayerEditorSubsystem->ToggleDataLayersIsDynamicallyLoadedInEditor(AllSelectedDataLayers);
+							bSuccess = DataLayerEditorSubsystem->ToggleDataLayersIsDynamicallyLoadedInEditor(AllSelectedDataLayers, true);
 						}
 						else
 						{
 							const FScopedTransaction Transaction(LOCTEXT("ToggleDataLayerIsDynamicallyLoadedInEditor", "Toggle Data Layer Dynamically Loaded In Editor Flag"));
-							bSuccess = DataLayerEditorSubsystem->ToggleDataLayerIsDynamicallyLoadedInEditor(DataLayer);
+							bSuccess = DataLayerEditorSubsystem->ToggleDataLayerIsDynamicallyLoadedInEditor(DataLayer, true);
 						}
 						if (!bSuccess)
 						{
diff --git a/Engine/Source/Editor/DataLayerEditor/Public/DataLayer/DataLayerEditorSubsystem.h b/Engine/Source/Editor/DataLayerEditor/Public/DataLayer/DataLayerEditorSubsystem.h
--- a/Engine/Source/Editor/DataLayerEditor/Public/DataLayer/DataLayerEditorSubsystem.h
+++ b/Engine/Source/Editor/DataLayerEditor/Public/DataLayer/DataLayerEditorSubsystem.h
@@ -478,9 +478,10 @@ public:
 	 * @param	bIsDynamicallyLoadedInEditor	The new value of the flag IsDynamicallyLoadedInEditor.
 	 *											If the DataLayer is DynamicallyLoaded, the Editor loading will consider this DataLayer to load or not an Actor.
 	 *											An Actor will not be loaded in the Editor if all its DataLayers are DynamicallyLoaded and not DynamicallyLoadedInEditor.
+	 * @param	bIsFromUserChange				If this change originates from a user change or not.
 	 */
 	UFUNCTION(BlueprintCallable, Category = DataLayers)
-	bool SetDataLayerIsDynamicallyLoadedInEditor(UDataLayer* DataLayer, const bool bIsDynamicallyLoadedInEditor);
+	bool SetDataLayerIsDynamicallyLoadedInEditor(UDataLayer* DataLayer, const bool bIsDynamicallyLoadedInEditor, const bool bIsFromUserChange);
 
 	/**
 	 * Changes the IsDynamicallyLoadedInEditor flag of the DataLayers to the provided state
@@ -489,25 +490,28 @@ public:
 	 * @param	bIsDynamicallyLoadedInEditor	The new value of the flag IsDynamicallyLoadedInEditor.
 	 *											If the DataLayer is DynamicallyLoaded, the Editor loading will consider this DataLayer to load or not an Actor.
 	 *											An Actor will not be loaded in the Editor if all its DataLayers are DynamicallyLoaded and not DynamicallyLoadedInEditor.
+	 * @param	bIsFromUserChange				If this change originates from a user change or not.
 	 */
 	UFUNCTION(BlueprintCallable, Category = DataLayers)
-	bool SetDataLayersIsDynamicallyLoadedInEditor(const TArray<UDataLayer*>& DataLayers, const bool bIsDynamicallyLoadedInEditor);
+	bool SetDataLayersIsDynamicallyLoadedInEditor(const TArray<UDataLayer*>& DataLayers, const bool bIsDynamicallyLoadedInEditor, const bool bIsFromUserChange);
 
 	/**
 	 * Toggles the DataLayer's IsDynamicallyLoadedInEditor flag
 	 *
-	 * @param DataLayer	The DataLayer to affect
+	 * @param	DataLayer						The DataLayer to affect
+	 * @param	bIsFromUserChange				If this change originates from a user change or not.
 	 */
 	UFUNCTION(BlueprintCallable, Category = DataLayers)
-	bool ToggleDataLayerIsDynamicallyLoadedInEditor(UDataLayer* DataLayer);
+	bool ToggleDataLayerIsDynamicallyLoadedInEditor(UDataLayer* DataLayer, const bool bIsFromUserChange);
 
 	/**
 	 * Toggles the IsDynamicallyLoadedInEditor flag of all of the DataLayers
 	 *
-	 * @param	DataLayers	The DataLayers to affect
+	 * @param	DataLayers						The DataLayers to affect
+	 * @param	bIsFromUserChange				If this change originates from a user change or not.
 	 */
 	UFUNCTION(BlueprintCallable, Category = DataLayers)
-	bool ToggleDataLayersIsDynamicallyLoadedInEditor(const TArray<UDataLayer*>& DataLayers);
+	bool ToggleDataLayersIsDynamicallyLoadedInEditor(const TArray<UDataLayer*>& DataLayers, const bool bIsFromUserChange);
 
 	/**
 	 * Set the visibility of all DataLayers to true
@@ -645,7 +649,7 @@ private:
 	void PostUndoRedo();
 
 	bool SetDataLayerIsDynamicallyLoadedInternal(UDataLayer* DataLayer, const bool bIsDynamicallyLoaded);
-	bool SetDataLayerIsDynamicallyLoadedInEditorInternal(UDataLayer* DataLayer, const bool bIsDynamicallyLoadedInEditor);
+	bool SetDataLayerIsDynamicallyLoadedInEditorInternal(UDataLayer* DataLayer, const bool bIsDynamicallyLoadedInEditor, const bool bIsFromUserChange);
 	bool RefreshWorldPartitionEditorCells();
 	void UpdateDataLayerEditorPerProjectUserSettings();
 
diff --git a/Engine/Source/Editor/UnrealEd/Private/WorldPartition/WorldPartitionBuilder.cpp b/Engine/Source/Editor/UnrealEd/Private/WorldPartition/WorldPartitionBuilder.cpp
--- a/Engine/Source/Editor/UnrealEd/Private/WorldPartition/WorldPartitionBuilder.cpp
+++ b/Engine/Source/Editor/UnrealEd/Private/WorldPartition/WorldPartitionBuilder.cpp
@@ -171,7 +171,7 @@ bool UWorldPartitionBuilder::Run(UWorld* World, FPackageSourceControlHelper& Pac
 			if (DataLayer->IsDynamicallyLoadedInEditor() != bLoadedInEditor)
 			{
 				bUpdateEditorCells = true;
-				DataLayer->SetIsDynamicallyLoadedInEditor(bLoadedInEditor);
+				DataLayer->SetIsDynamicallyLoadedInEditor(bLoadedInEditor, false);
 				if (RequiresCommandletRendering() && bLoadedInEditor)
 				{
 					DataLayer->SetIsInitiallyVisible(true);
diff --git a/Engine/Source/Runtime/Engine/Private/WorldPartition/DataLayer/DataLayer.cpp b/Engine/Source/Runtime/Engine/Private/WorldPartition/DataLayer/DataLayer.cpp
--- a/Engine/Source/Runtime/Engine/Private/WorldPartition/DataLayer/DataLayer.cpp
+++ b/Engine/Source/Runtime/Engine/Private/WorldPartition/DataLayer/DataLayer.cpp
@@ -117,12 +117,13 @@ void UDataLayer::SetIsDynamicallyLoaded(bool bInIsDynamicallyLoaded)
 	}
 }
 
-void UDataLayer::SetIsDynamicallyLoadedInEditor(bool bInIsDynamicallyLoadedInEditor)
+void UDataLayer::SetIsDynamicallyLoadedInEditor(bool bInIsDynamicallyLoadedInEditor, bool bInFromUserChange)
 {
 	if (bIsDynamicallyLoadedInEditor != bInIsDynamicallyLoadedInEditor)
 	{
 		Modify(false);
 		bIsDynamicallyLoadedInEditor = bInIsDynamicallyLoadedInEditor;
+		bIsDynamicallyLoadedInEditorChangedByUserOperation |= bInFromUserChange;
 	}
 }
 
diff --git a/Engine/Source/Runtime/Engine/Private/WorldPartition/DataLayer/WorldDataLayers.cpp b/Engine/Source/Runtime/Engine/Private/WorldPartition/DataLayer/WorldDataLayers.cpp
--- a/Engine/Source/Runtime/Engine/Private/WorldPartition/DataLayer/WorldDataLayers.cpp
+++ b/Engine/Source/Runtime/Engine/Private/WorldPartition/DataLayer/WorldDataLayers.cpp
@@ -414,7 +414,7 @@ void AWorldDataLayers::PostLoad()
 	// Setup defaults before overriding with user settings
 	for (UDataLayer* DataLayer : WorldDataLayers)
 	{
-		DataLayer->SetIsDynamicallyLoadedInEditor(DataLayer->IsInitiallyLoadedInEditor());
+		DataLayer->SetIsDynamicallyLoadedInEditor(DataLayer->IsInitiallyLoadedInEditor(), false);
 	}
 
 	// Initialize DataLayer's IsDynamicallyLoadedInEditor based on DataLayerEditorPerProjectUserSettings
@@ -423,7 +423,7 @@ void AWorldDataLayers::PostLoad()
 	{
 		if (UDataLayer* DataLayer = const_cast<UDataLayer*>(GetDataLayerFromName(DataLayerName)))
 		{
-			DataLayer->SetIsDynamicallyLoadedInEditor(false);
+			DataLayer->SetIsDynamicallyLoadedInEditor(false, false);
 		}
 	}
 
@@ -432,7 +432,7 @@ void AWorldDataLayers::PostLoad()
 	{
 		if (UDataLayer* DataLayer = const_cast<UDataLayer*>(GetDataLayerFromName(DataLayerName)))
 		{
-			DataLayer->SetIsDynamicallyLoadedInEditor(true);
+			DataLayer->SetIsDynamicallyLoadedInEditor(true, false);
 }
 	}
 
diff --git a/Engine/Source/Runtime/Engine/Public/WorldPartition/DataLayer/DataLayer.h b/Engine/Source/Runtime/Engine/Public/WorldPartition/DataLayer/DataLayer.h
--- a/Engine/Source/Runtime/Engine/Public/WorldPartition/DataLayer/DataLayer.h
+++ b/Engine/Source/Runtime/Engine/Public/WorldPartition/DataLayer/DataLayer.h
@@ -32,10 +32,12 @@ public:
 	void SetVisible(bool bInIsVisible);
 	void SetIsInitiallyVisible(bool bInIsInitiallyVisible);
 	void SetIsDynamicallyLoaded(bool bInIsDynamicallyLoaded);
-	void SetIsDynamicallyLoadedInEditor(bool bInIsDynamicallyLoadedInEditor);
+	void SetIsDynamicallyLoadedInEditor(bool bInIsDynamicallyLoadedInEditor, bool bInFromUserChange);
 	void SetIsLocked(bool bInIsLocked) { bIsLocked = bInIsLocked; }
 
 	bool IsDynamicallyLoadedInEditor() const { return bIsDynamicallyLoadedInEditor; }
+	bool IsDynamicallyLoadedInEditorChangedByUserOperation() const { return bIsDynamicallyLoadedInEditorChangedByUserOperation; }
+	void ClearDynamicallyLoadedInEditorChangedByUserOperation() { bIsDynamicallyLoadedInEditorChangedByUserOperation = false; }
 	bool IsInitiallyLoadedInEditor() const { return bIsInitiallyLoadedInEditor; }
 
 	static FText GetDataLayerText(const UDataLayer* InDataLayer);
@@ -92,6 +94,10 @@ private:
 	UPROPERTY(Transient)
 	uint32 bIsDynamicallyLoadedInEditor : 1;
 
+	/** Whether this data layer editor visibility was changed by a user operation */
+	UPROPERTY(Transient)
+	uint32 bIsDynamicallyLoadedInEditorChangedByUserOperation : 1;
+
 	/** Whether this data layer is locked, which means the user can't change actors assignation, remove or rename it */
 	UPROPERTY()
 	uint32 bIsLocked : 1;
diff --git a/Engine/Source/Runtime/MovieSceneTracks/Private/Systems/MovieSceneDataLayerSystem.cpp b/Engine/Source/Runtime/MovieSceneTracks/Private/Systems/MovieSceneDataLayerSystem.cpp
--- a/Engine/Source/Runtime/MovieSceneTracks/Private/Systems/MovieSceneDataLayerSystem.cpp
+++ b/Engine/Source/Runtime/MovieSceneTracks/Private/Systems/MovieSceneDataLayerSystem.cpp
@@ -406,7 +406,7 @@ void FDesiredLayerStates::ApplyInEditor(FPreAnimatedDataLayerStorage* PreAnimate
 	if (DatalayersNeedingLoad.Num() > 0)
 	{
 		// This blocks while we load data layers
-		SubSystem->SetDataLayersIsDynamicallyLoadedInEditor(DatalayersNeedingLoad, true);
+		SubSystem->SetDataLayersIsDynamicallyLoadedInEditor(DatalayersNeedingLoad, true, false);
 	}
 	if (DatalayersNeedingShow.Num() > 0)
 	{
