Minor optimization in data layers slate dialog display when in PIE.

diff --git a/Engine/Source/Editor/DataLayerEditor/Private/DataLayer/SDataLayerTreeLabel.cpp b/Engine/Source/Editor/DataLayerEditor/Private/DataLayer/SDataLayerTreeLabel.cpp
--- a/Engine/Source/Editor/DataLayerEditor/Private/DataLayer/SDataLayerTreeLabel.cpp
+++ b/Engine/Source/Editor/DataLayerEditor/Private/DataLayer/SDataLayerTreeLabel.cpp
@@ -94,7 +94,7 @@ FText SDataLayerTreeLabel::GetDisplayText() const
 	if (DataLayer && DataLayer->IsDynamicallyLoaded() && DataLayer->GetWorld()->IsPlayInEditor())
 	{
 		const UDataLayerSubsystem* DataLayerSubsystem = DataLayer->GetWorld()->GetSubsystem<UDataLayerSubsystem>();
-		DataLayerStateText = FText::Format(LOCTEXT("DataLayerActive", " ({0})"), StaticEnum<EDataLayerState>()->GetDisplayNameTextByValue((int64)DataLayerSubsystem->GetDataLayerState(DataLayer)));
+		DataLayerStateText = FText::Format(LOCTEXT("DataLayerActive", " ({0})"), FTextStringHelper::CreateFromBuffer(GetDataLayerStateName(DataLayerSubsystem->GetDataLayerState(DataLayer))));
 	}
 	
 	static const FText DataLayerDeleted = LOCTEXT("DataLayerLabelForMissingDataLayer", "(Deleted Data Layer)");
diff --git a/Engine/Source/Runtime/Engine/Public/WorldPartition/DataLayer/DataLayer.h b/Engine/Source/Runtime/Engine/Public/WorldPartition/DataLayer/DataLayer.h
--- a/Engine/Source/Runtime/Engine/Public/WorldPartition/DataLayer/DataLayer.h
+++ b/Engine/Source/Runtime/Engine/Public/WorldPartition/DataLayer/DataLayer.h
@@ -19,6 +19,18 @@ enum class EDataLayerState : uint8
 	Activated
 };
 
+const inline TCHAR* GetDataLayerStateName(EDataLayerState DataLayerState)
+{
+	switch(DataLayerState)
+	{
+	case EDataLayerState::Unloaded: return TEXT("Unloaded");
+	case EDataLayerState::Loaded: return TEXT("Loaded");
+	case EDataLayerState::Activated: return TEXT("Activated");
+	default: check(0);
+	}
+	return TEXT("Invalid");
+}
+
 static_assert(EDataLayerState::Unloaded < EDataLayerState::Loaded && EDataLayerState::Loaded < EDataLayerState::Activated, "Streaming Query code is dependent on this being true");
 
 UCLASS(Config = Engine, PerObjectConfig, Within = WorldDataLayers, BlueprintType)
