misc non-unity compilation fixes

diff --git a/Engine/Source/Developer/IOS/IOSPlatformEditor/Private/IOSTargetSettingsCustomization.cpp b/Engine/Source/Developer/IOS/IOSPlatformEditor/Private/IOSTargetSettingsCustomization.cpp
--- a/Engine/Source/Developer/IOS/IOSPlatformEditor/Private/IOSTargetSettingsCustomization.cpp
+++ b/Engine/Source/Developer/IOS/IOSPlatformEditor/Private/IOSTargetSettingsCustomization.cpp
@@ -36,6 +36,7 @@
 #include "Interfaces/ITargetPlatform.h"
 #include "Interfaces/ITargetPlatformModule.h"
 #include "IOSPlatformEditorModule.h"
+#include "HAL/PlatformFileManager.h"
 #include "UEFreeImage.h"
 
 #define LOCTEXT_NAMESPACE "IOSTargetSettings"
diff --git a/Engine/Source/Editor/DataLayerEditor/Private/DataLayer/DataLayerHierarchy.cpp b/Engine/Source/Editor/DataLayerEditor/Private/DataLayer/DataLayerHierarchy.cpp
--- a/Engine/Source/Editor/DataLayerEditor/Private/DataLayer/DataLayerHierarchy.cpp
+++ b/Engine/Source/Editor/DataLayerEditor/Private/DataLayer/DataLayerHierarchy.cpp
@@ -9,9 +9,11 @@
 #include "DataLayer/DataLayerEditorSubsystem.h"
 #include "WorldPartition/DataLayer/DataLayer.h"
 #include "WorldPartition/DataLayer/WorldDataLayers.h"
+#include "Engine/Engine.h"
 #include "Engine/World.h"
 #include "EngineUtils.h"
 #include "Editor.h"
+#include "Modules/ModuleManager.h"
 #include "WorldPartition/WorldPartitionHelpers.h"
 #include "WorldPartition/IWorldPartitionEditorModule.h"
 
diff --git a/Engine/Source/Editor/Kismet/Private/UserDefinedEnumEditor.cpp b/Engine/Source/Editor/Kismet/Private/UserDefinedEnumEditor.cpp
--- a/Engine/Source/Editor/Kismet/Private/UserDefinedEnumEditor.cpp
+++ b/Engine/Source/Editor/Kismet/Private/UserDefinedEnumEditor.cpp
@@ -19,6 +19,7 @@
 #include "DragAndDrop/DecoratedDragDropOp.h"
 #include "SEditorHeaderButton.h"
 #include "Styling/ToolBarStyle.h"
+#include "Framework/MultiBox/MultiBoxBuilder.h"
 
 #include "PropertyCustomizationHelpers.h"
 #include "Widgets/SToolTip.h"
diff --git a/Engine/Source/Editor/SceneOutliner/Private/ActorHierarchy.cpp b/Engine/Source/Editor/SceneOutliner/Private/ActorHierarchy.cpp
--- a/Engine/Source/Editor/SceneOutliner/Private/ActorHierarchy.cpp
+++ b/Engine/Source/Editor/SceneOutliner/Private/ActorHierarchy.cpp
@@ -18,6 +18,7 @@
 #include "EditorFolderUtils.h"
 #include "LevelInstance/LevelInstanceActor.h"
 #include "LevelInstance/LevelInstanceSubsystem.h"
+#include "Modules/ModuleManager.h"
 #include "WorldPartition/WorldPartitionHelpers.h"
 #include "WorldPartition/WorldPartitionActorDesc.h"
 #include "WorldPartition/IWorldPartitionEditorModule.h"
diff --git a/Engine/Source/Runtime/Engine/Private/ActorPartition/PartitionActor.cpp b/Engine/Source/Runtime/Engine/Private/ActorPartition/PartitionActor.cpp
--- a/Engine/Source/Runtime/Engine/Private/ActorPartition/PartitionActor.cpp
+++ b/Engine/Source/Runtime/Engine/Private/ActorPartition/PartitionActor.cpp
@@ -6,6 +6,8 @@
 #include "Components/BoxComponent.h"
 #include "WorldPartition/WorldPartitionActorDesc.h"
 #include "WorldPartition/ActorPartition/PartitionActorDesc.h"
+#include "GameFramework/WorldSettings.h"
+#include "Engine/World.h"
 #endif
 
 #define LOCTEXT_NAMESPACE "PartitionActor"
diff --git a/Engine/Source/Runtime/Engine/Private/Animation/MotionTrajectoryTypes.cpp b/Engine/Source/Runtime/Engine/Private/Animation/MotionTrajectoryTypes.cpp
--- a/Engine/Source/Runtime/Engine/Private/Animation/MotionTrajectoryTypes.cpp
+++ b/Engine/Source/Runtime/Engine/Private/Animation/MotionTrajectoryTypes.cpp
@@ -3,6 +3,7 @@
 #include "Animation/MotionTrajectoryTypes.h"
 #include "Algo/AllOf.h"
 #include "DrawDebugHelpers.h"
+#include "Animation/AnimTypes.h"
 
 #if ENABLE_ANIM_DEBUG
 constexpr int32 DebugTrajectorySampleDisable = 0;
diff --git a/Engine/Source/Runtime/Renderer/Private/OIT/OIT.h b/Engine/Source/Runtime/Renderer/Private/OIT/OIT.h
--- a/Engine/Source/Runtime/Renderer/Private/OIT/OIT.h
+++ b/Engine/Source/Runtime/Renderer/Private/OIT/OIT.h
@@ -7,6 +7,7 @@
 #include "RenderGraphResources.h"
 #include "RenderResource.h"
 
+class FViewInfo;
 class FSortedIndexBuffer;
 struct FMeshBatchElementDynamicIndexBuffer;
 
