Fixed missing case in FDataLayerHierarchy::CreateParentItem

diff --git a/Engine/Source/Editor/DataLayerEditor/Private/DataLayer/DataLayerHierarchy.cpp b/Engine/Source/Editor/DataLayerEditor/Private/DataLayer/DataLayerHierarchy.cpp
--- a/Engine/Source/Editor/DataLayerEditor/Private/DataLayer/DataLayerHierarchy.cpp
+++ b/Engine/Source/Editor/DataLayerEditor/Private/DataLayer/DataLayerHierarchy.cpp
@@ -208,9 +208,15 @@ FSceneOutlinerTreeItemPtr FDataLayerHierarchy::FindParent(const ISceneOutlinerTr
 
 FSceneOutlinerTreeItemPtr FDataLayerHierarchy::CreateParentItem(const FSceneOutlinerTreeItemPtr& Item) const
 {
-	if (Item->IsA<FDataLayerTreeItem>())
+	if (FDataLayerTreeItem* DataLayerTreeItem = Item->CastTo<FDataLayerTreeItem>())
 	{
-		return nullptr;
+		if (UDataLayer* DataLayer = DataLayerTreeItem->GetDataLayer())
+		{
+			if (UDataLayer* ParentDataLayer = DataLayer->GetParent())
+			{
+				return Mode->CreateItemFor<FDataLayerTreeItem>(ParentDataLayer, true);
+			}
+		}
 	}
 	else if (FDataLayerActorTreeItem* DataLayerActorTreeItem = Item->CastTo<FDataLayerActorTreeItem>())
 	{
