Data Layer is considered not visible when its marked as unloaded in editor.

diff --git a/Engine/Source/Editor/DataLayerEditor/Private/DataLayer/DataLayerEditorSubsystem.cpp b/Engine/Source/Editor/DataLayerEditor/Private/DataLayer/DataLayerEditorSubsystem.cpp
--- a/Engine/Source/Editor/DataLayerEditor/Private/DataLayer/DataLayerEditorSubsystem.cpp
+++ b/Engine/Source/Editor/DataLayerEditor/Private/DataLayer/DataLayerEditorSubsystem.cpp
@@ -776,9 +776,16 @@ bool UDataLayerEditorSubsystem::SetDataLayerIsLoadedInEditorInternal(UDataLayer*
 	check(DataLayer);
 	if (DataLayer->IsLoadedInEditor() != bIsLoadedInEditor)
 	{
+		const bool bWasVisible = DataLayer->IsEffectiveVisible();
+
 		DataLayer->Modify(false);
 		DataLayer->SetIsLoadedInEditor(bIsLoadedInEditor, /*bFromUserChange*/bIsFromUserChange);
 		DataLayerChanged.Broadcast(EDataLayerAction::Modify, DataLayer, "bIsLoadedInEditor");
+
+		if (DataLayer->IsEffectiveVisible() != bWasVisible)
+		{
+			UpdateAllActorsVisibility(true, true);
+		}
 		return true;
 	}
 	return false;
diff --git a/Engine/Source/Editor/DataLayerEditor/Private/DataLayer/DataLayerOutlinerIsVisibleColumn.cpp b/Engine/Source/Editor/DataLayerEditor/Private/DataLayer/DataLayerOutlinerIsVisibleColumn.cpp
--- a/Engine/Source/Editor/DataLayerEditor/Private/DataLayer/DataLayerOutlinerIsVisibleColumn.cpp
+++ b/Engine/Source/Editor/DataLayerEditor/Private/DataLayer/DataLayerOutlinerIsVisibleColumn.cpp
@@ -32,7 +32,7 @@ protected:
 			const UDataLayer* DataLayer = DataLayerTreeItem->GetDataLayer();
 			const UDataLayer* ParentDataLayer = DataLayer ? DataLayer->GetParent() : nullptr;
 			const bool bIsParentVisible = ParentDataLayer ? ParentDataLayer->IsEffectiveVisible() : true;
-			return bIsParentVisible && DataLayer && DataLayer->GetWorld() && !DataLayer->GetWorld()->IsPlayInEditor();
+			return bIsParentVisible && DataLayer && DataLayer->GetWorld() && !DataLayer->GetWorld()->IsPlayInEditor() && DataLayer->IsEffectiveLoadedInEditor();
 		}
 		return false;
 	}
diff --git a/Engine/Source/Runtime/Engine/Private/WorldPartition/DataLayer/DataLayer.cpp b/Engine/Source/Runtime/Engine/Private/WorldPartition/DataLayer/DataLayer.cpp
--- a/Engine/Source/Runtime/Engine/Private/WorldPartition/DataLayer/DataLayer.cpp
+++ b/Engine/Source/Runtime/Engine/Private/WorldPartition/DataLayer/DataLayer.cpp
@@ -93,7 +93,7 @@ bool UDataLayer::IsEffectiveVisible() const
 		bResult = bResult && ParentDataLayer->IsVisible();
 		ParentDataLayer = ParentDataLayer->GetParent();
 	}
-	return bResult;
+	return bResult && IsEffectiveLoadedInEditor();
 #else
 	return false;
 #endif
