Fixed bug in Data Layer outliner where newly created actors, once saved, were showing a second unloaded version of the same actor.

diff --git a/Engine/Source/Editor/DataLayerEditor/Private/DataLayer/DataLayerHierarchy.cpp b/Engine/Source/Editor/DataLayerEditor/Private/DataLayer/DataLayerHierarchy.cpp
--- a/Engine/Source/Editor/DataLayerEditor/Private/DataLayer/DataLayerHierarchy.cpp
+++ b/Engine/Source/Editor/DataLayerEditor/Private/DataLayer/DataLayerHierarchy.cpp
@@ -405,7 +405,7 @@ void FDataLayerHierarchy::OnLoadedActorRemoved(AActor& InActor)
 
 void FDataLayerHierarchy::OnActorDescAdded(FWorldPartitionActorDesc* InActorDesc)
 {
-	if (!bShowingUnloadedActors || (InActorDesc == nullptr))
+	if (!bShowingUnloadedActors || !InActorDesc || InActorDesc->IsLoaded(true))
 	{
 		return;
 	}
