- Added Data Layer highlighting based on editor selection
- Added Data Layer / Actor filtering based on editor selection

diff --git a/Engine/Source/Editor/DataLayerEditor/Private/DataLayer/DataLayerEditorSubsystem.cpp b/Engine/Source/Editor/DataLayerEditor/Private/DataLayer/DataLayerEditorSubsystem.cpp
--- a/Engine/Source/Editor/DataLayerEditor/Private/DataLayer/DataLayerEditorSubsystem.cpp
+++ b/Engine/Source/Editor/DataLayerEditor/Private/DataLayer/DataLayerEditorSubsystem.cpp
@@ -35,6 +35,7 @@ private:
 	void OnPostUndoRedo() { DataLayerEditorSubsystem->PostUndoRedo(); }
 	void OnObjectPostEditChange(UObject* Object, FPropertyChangedEvent& PropertyChangedEvent);
 	void OnLevelActorsAdded(AActor* InActor) { DataLayerEditorSubsystem->InitializeNewActorDataLayers(InActor); }
+	void OnLevelSelectionChanged(UObject* InObject) { DataLayerEditorSubsystem->OnSelectionChanged(); }
 
 	UDataLayerEditorSubsystem* DataLayerEditorSubsystem;
 	bool bIsInitialized;
@@ -61,6 +62,8 @@ void FDataLayersBroadcast::Deinitialize()
 		FEditorDelegates::PostUndoRedo.RemoveAll(this);
 		FCoreUObjectDelegates::OnObjectPropertyChanged.RemoveAll(this);
 		GEngine->OnLevelActorAdded().RemoveAll(this);
+		USelection::SelectionChangedEvent.RemoveAll(this);
+		USelection::SelectObjectEvent.RemoveAll(this);
 	}
 }
 
@@ -73,6 +76,8 @@ void FDataLayersBroadcast::Initialize()
 		FEditorDelegates::PostUndoRedo.AddRaw(this, &FDataLayersBroadcast::OnPostUndoRedo);
 		FCoreUObjectDelegates::OnObjectPropertyChanged.AddRaw(this, &FDataLayersBroadcast::OnObjectPostEditChange);
 		GEngine->OnLevelActorAdded().AddRaw(this, &FDataLayersBroadcast::OnLevelActorsAdded);
+		USelection::SelectionChangedEvent.AddRaw(this, &FDataLayersBroadcast::OnLevelSelectionChanged);
+		USelection::SelectObjectEvent.AddRaw(this, &FDataLayersBroadcast::OnLevelSelectionChanged);
 	}
 }
 
@@ -162,19 +167,19 @@ void UDataLayerEditorSubsystem::EditorMapChange()
 	{
 		World->PersistentLevel->OnLoadedActorAddedToLevelEvent.AddLambda([this](AActor& InActor) { InitializeNewActorDataLayers(&InActor); });
 	}
-	DataLayerChanged.Broadcast(EDataLayerAction::Reset, NULL, NAME_None);
+	BroadcastDataLayerChanged(EDataLayerAction::Reset, NULL, NAME_None);
 	UpdateAllActorsVisibility(true, true);
 }
 
 void UDataLayerEditorSubsystem::EditorRefreshDataLayerBrowser()
 {
-	DataLayerChanged.Broadcast(EDataLayerAction::Reset, NULL, NAME_None);
+	BroadcastDataLayerChanged(EDataLayerAction::Reset, NULL, NAME_None);
 	UpdateAllActorsVisibility(false, false);
 }
 
 void UDataLayerEditorSubsystem::PostUndoRedo()
 {
-	DataLayerChanged.Broadcast(EDataLayerAction::Reset, NULL, NAME_None);
+	BroadcastDataLayerChanged(EDataLayerAction::Reset, NULL, NAME_None);
 	UpdateAllActorsVisibility(true, true);
 }
 
@@ -215,7 +220,7 @@ bool UDataLayerEditorSubsystem::SetParentDataLayer(UDataLayer* DataLayer, UDataL
 	{
 		const bool bIsLoaded = DataLayer->IsEffectiveLoadedInEditor();
 		DataLayer->SetParent(ParentDataLayer);
-		DataLayerChanged.Broadcast(EDataLayerAction::Reset, NULL, NAME_None);
+		BroadcastDataLayerChanged(EDataLayerAction::Reset, NULL, NAME_None);
 		UpdateAllActorsVisibility(true, true);
 		if (bIsLoaded != DataLayer->IsEffectiveLoadedInEditor())
 		{
@@ -268,7 +273,7 @@ bool UDataLayerEditorSubsystem::AddActorsToDataLayers(const TArray<AActor*>& Act
 				if (Actor->AddDataLayer(DataLayer))
 				{
 					bActorWasModified = true;
-					ActorDataLayersChanged.Broadcast(Actor);
+					BroadcastActorDataLayersChanged(Actor);
 				}
 			}
 
@@ -306,9 +311,9 @@ bool UDataLayerEditorSubsystem::RemoveActorsFromAllDataLayers(const TArray<AActo
 		{
 			for (const UDataLayer* DataLayer : ModifiedDataLayers)
 			{
-				DataLayerChanged.Broadcast(EDataLayerAction::Modify, DataLayer, NAME_None);
+				BroadcastDataLayerChanged(EDataLayerAction::Modify, DataLayer, NAME_None);
 			}
-			ActorDataLayersChanged.Broadcast(Actor);
+			BroadcastActorDataLayersChanged(Actor);
 
 			// Update general actor visibility
 			bool bActorModified = false;
@@ -363,8 +368,8 @@ bool UDataLayerEditorSubsystem::RemoveActorsFromDataLayers(const TArray<AActor*>
 			if (Actor->RemoveDataLayer(DataLayer))
 			{
 				bActorWasModified = true;
-				DataLayerChanged.Broadcast(EDataLayerAction::Modify, DataLayer, NAME_None);
-				ActorDataLayersChanged.Broadcast(Actor);
+				BroadcastDataLayerChanged(EDataLayerAction::Modify, DataLayer, NAME_None);
+				BroadcastActorDataLayersChanged(Actor);
 			}
 		}
 
@@ -718,7 +723,7 @@ void UDataLayerEditorSubsystem::SetDataLayersVisibility(const TArray<UDataLayer*
 		{
 			DataLayer->Modify(/*bAlswaysMarkDirty*/false);
 			DataLayer->SetVisible(bIsVisible);
-			DataLayerChanged.Broadcast(EDataLayerAction::Modify, DataLayer, "bIsVisible");
+			BroadcastDataLayerChanged(EDataLayerAction::Modify, DataLayer, "bIsVisible");
 			bChangeOccurred = true;
 		}
 	}
@@ -746,7 +751,7 @@ void UDataLayerEditorSubsystem::ToggleDataLayersVisibility(const TArray<UDataLay
 	{
 		DataLayer->Modify();
 		DataLayer->SetVisible(!DataLayer->IsVisible());
-		DataLayerChanged.Broadcast(EDataLayerAction::Modify, DataLayer, "bIsVisible");
+		BroadcastDataLayerChanged(EDataLayerAction::Modify, DataLayer, "bIsVisible");
 	}
 
 	UpdateAllActorsVisibility(true, true);
@@ -762,7 +767,7 @@ void UDataLayerEditorSubsystem::MakeAllDataLayersVisible()
 			{
 				DataLayer->Modify();
 				DataLayer->SetVisible(true);
-				DataLayerChanged.Broadcast(EDataLayerAction::Modify, DataLayer, "bIsVisible");
+				BroadcastDataLayerChanged(EDataLayerAction::Modify, DataLayer, "bIsVisible");
 			}
 			return true;
 		});
@@ -780,7 +785,7 @@ bool UDataLayerEditorSubsystem::SetDataLayerIsLoadedInEditorInternal(UDataLayer*
 
 		DataLayer->Modify(false);
 		DataLayer->SetIsLoadedInEditor(bIsLoadedInEditor, /*bFromUserChange*/bIsFromUserChange);
-		DataLayerChanged.Broadcast(EDataLayerAction::Modify, DataLayer, "bIsLoadedInEditor");
+		BroadcastDataLayerChanged(EDataLayerAction::Modify, DataLayer, "bIsLoadedInEditor");
 
 		if (DataLayer->IsEffectiveVisible() != bWasVisible)
 		{
@@ -893,7 +898,7 @@ UDataLayer* UDataLayerEditorSubsystem::CreateDataLayer()
 {
 	AWorldDataLayers* WorldDataLayers = GetWorldDataLayers(/*bCreateIfNotFound*/true);
 	UDataLayer* NewDataLayer = WorldDataLayers->CreateDataLayer();
-	DataLayerChanged.Broadcast(EDataLayerAction::Add, NewDataLayer, NAME_None);
+	BroadcastDataLayerChanged(EDataLayerAction::Add, NewDataLayer, NAME_None);
 	return NewDataLayer;
 }
 
@@ -903,7 +908,7 @@ void UDataLayerEditorSubsystem::DeleteDataLayers(const TArray<UDataLayer*>& Data
 	{
 		if (WorldDataLayers->RemoveDataLayers(DataLayersToDelete))
 		{
-			DataLayerChanged.Broadcast(EDataLayerAction::Delete, NULL, NAME_None);
+			BroadcastDataLayerChanged(EDataLayerAction::Delete, NULL, NAME_None);
 		}
 	}
 }
@@ -914,7 +919,7 @@ void UDataLayerEditorSubsystem::DeleteDataLayer(UDataLayer* DataLayerToDelete)
 	{
 		if (WorldDataLayers->RemoveDataLayer(DataLayerToDelete))
 		{
-			DataLayerChanged.Broadcast(EDataLayerAction::Delete, NULL, NAME_None);
+			BroadcastDataLayerChanged(EDataLayerAction::Delete, NULL, NAME_None);
 		}
 	}
 }
@@ -929,9 +934,41 @@ bool UDataLayerEditorSubsystem::RenameDataLayer(UDataLayer* DataLayer, const FNa
 			FName UniqueNewDataLayerLabel = WorldDataLayers->GenerateUniqueDataLayerLabel(DataLayerLabelSanitized);
 			DataLayer->Modify();
 			DataLayer->SetDataLayerLabel(UniqueNewDataLayerLabel);
-			DataLayerChanged.Broadcast(EDataLayerAction::Rename, DataLayer, "DataLayerLabel");
+			BroadcastDataLayerChanged(EDataLayerAction::Rename, DataLayer, "DataLayerLabel");
 			return true;
 		}
 	}
 	return false;
 }
+
+void UDataLayerEditorSubsystem::BroadcastActorDataLayersChanged(const TWeakObjectPtr<AActor>& ChangedActor)
+{
+	RebuildSelectedDataLayersFromEditorSelection();
+	ActorDataLayersChanged.Broadcast(ChangedActor);
+}
+
+void UDataLayerEditorSubsystem::BroadcastDataLayerChanged(const EDataLayerAction Action, const TWeakObjectPtr<const UDataLayer>& ChangedDataLayer, const FName& ChangedProperty)
+{
+	RebuildSelectedDataLayersFromEditorSelection();
+	DataLayerChanged.Broadcast(Action, ChangedDataLayer, ChangedProperty);
+}
+
+void UDataLayerEditorSubsystem::OnSelectionChanged()
+{
+	RebuildSelectedDataLayersFromEditorSelection();
+}
+
+void UDataLayerEditorSubsystem::RebuildSelectedDataLayersFromEditorSelection()
+{
+	SelectedDataLayersFromEditorSelection.Reset();
+
+	TArray<AActor*> Actors;
+	GEditor->GetSelectedActors()->GetSelectedObjects<AActor>(Actors);
+	for (const AActor* Actor : Actors)
+	{
+		for (const UDataLayer* DataLayer : Actor->GetDataLayerObjects())
+		{
+			SelectedDataLayersFromEditorSelection.Add(DataLayer);
+		}
+	}
+}
diff --git a/Engine/Source/Editor/DataLayerEditor/Private/DataLayer/DataLayerHierarchy.cpp b/Engine/Source/Editor/DataLayerEditor/Private/DataLayer/DataLayerHierarchy.cpp
--- a/Engine/Source/Editor/DataLayerEditor/Private/DataLayer/DataLayerHierarchy.cpp
+++ b/Engine/Source/Editor/DataLayerEditor/Private/DataLayer/DataLayerHierarchy.cpp
@@ -25,10 +25,12 @@ TUniquePtr<FDataLayerHierarchy> FDataLayerHierarchy::Create(FDataLayerMode* Mode
 FDataLayerHierarchy::FDataLayerHierarchy(FDataLayerMode* Mode, const TWeakObjectPtr<UWorld>& World)
 	: ISceneOutlinerHierarchy(Mode)
 	, RepresentingWorld(World)
-	, bShowingEditorDataLayers(true)
-	, bShowingRuntimeDataLayers(true)
-	, bShowingDataLayerActors(true)
-	, bShowingUnloadedActors(true)
+	, bShowEditorDataLayers(true)
+	, bShowRuntimeDataLayers(true)
+	, bShowDataLayerActors(true)
+	, bShowUnloadedActors(true)
+	, bShowOnlySelectedActors(false)
+	, bHighlightSelectedDataLayers(false)
 {
 	if (GEngine)
 	{
@@ -94,6 +96,38 @@ FDataLayerHierarchy::~FDataLayerHierarchy()
 	FWorldDelegates::LevelRemovedFromWorld.RemoveAll(this);
 }
 
+
+bool FDataLayerHierarchy::IsDataLayerPartOfSelection(const UDataLayer* DataLayer) const
+{
+	if (!bShowOnlySelectedActors)
+	{
+		return true;
+	}
+
+	if (UDataLayerEditorSubsystem::Get()->DoesDataLayerContainSelectedActors(DataLayer))
+	{
+		return true;
+	}
+
+	bool bFoundSelected = false;
+	DataLayer->ForEachChild([this, &bFoundSelected](const UDataLayer* Child)
+	{
+		bFoundSelected = IsDataLayerPartOfSelection(Child);
+		return !bFoundSelected; // Continue iterating if not found
+	});
+	return bFoundSelected;
+};
+
+FSceneOutlinerTreeItemPtr FDataLayerHierarchy::CreateDataLayerTreeItem(UDataLayer* InDataLayer, bool bInForce) const
+{
+	FSceneOutlinerTreeItemPtr Item = Mode->CreateItemFor<FDataLayerTreeItem>(InDataLayer, bInForce);
+	if (FDataLayerTreeItem* DataLayerTreeItem = Item ? Item->CastTo<FDataLayerTreeItem>() : nullptr)
+	{
+		DataLayerTreeItem->SetIsHighlightedIfSelected(bHighlightSelectedDataLayers);
+	}
+	return Item;
+}
+
 void FDataLayerHierarchy::CreateItems(TArray<FSceneOutlinerTreeItemPtr>& OutItems) const
 {
 	const AWorldDataLayers* WorldDataLayers = RepresentingWorld.Get()->GetWorldDataLayers();
@@ -105,14 +139,14 @@ void FDataLayerHierarchy::CreateItems(TArray<FSceneOutlinerTreeItemPtr>& OutItem
 	auto IsDataLayerShown = [this](const UDataLayer* DataLayer)
 	{
 		const bool bIsRuntimeDataLayer = DataLayer->IsRuntime();
-		return (bIsRuntimeDataLayer && bShowingRuntimeDataLayers) || (!bIsRuntimeDataLayer && bShowingEditorDataLayers);
+		return ((bIsRuntimeDataLayer && bShowRuntimeDataLayers) || (!bIsRuntimeDataLayer && bShowEditorDataLayers)) && IsDataLayerPartOfSelection(DataLayer);
 	};
 
 	WorldDataLayers->ForEachDataLayer([this, &OutItems, IsDataLayerShown](UDataLayer* DataLayer)
 	{
 		if (IsDataLayerShown(DataLayer))
 		{
-			if (FSceneOutlinerTreeItemPtr DataLayerItem = Mode->CreateItemFor<FDataLayerTreeItem>(DataLayer))
+			if (FSceneOutlinerTreeItemPtr DataLayerItem = CreateDataLayerTreeItem(DataLayer))
 			{
 				OutItems.Add(DataLayerItem);
 			}
@@ -120,7 +154,7 @@ void FDataLayerHierarchy::CreateItems(TArray<FSceneOutlinerTreeItemPtr>& OutItem
 		return true;
 	});
 
-	if (bShowingDataLayerActors)
+	if (bShowDataLayerActors)
 	{
 		for (AActor* Actor : FActorRange(RepresentingWorld.Get()))
 		{
@@ -139,7 +173,7 @@ void FDataLayerHierarchy::CreateItems(TArray<FSceneOutlinerTreeItemPtr>& OutItem
 			}
 		}
 
-		if (bShowingUnloadedActors)
+		if (bShowUnloadedActors)
 		{
 			if (UWorldPartition* const WorldPartition = RepresentingWorld.Get()->GetWorldPartition())
 			{
@@ -214,7 +248,7 @@ FSceneOutlinerTreeItemPtr FDataLayerHierarchy::CreateParentItem(const FSceneOutl
 		{
 			if (UDataLayer* ParentDataLayer = DataLayer->GetParent())
 			{
-				return Mode->CreateItemFor<FDataLayerTreeItem>(ParentDataLayer, true);
+				return CreateDataLayerTreeItem(ParentDataLayer, true);
 			}
 		}
 	}
@@ -222,14 +256,14 @@ FSceneOutlinerTreeItemPtr FDataLayerHierarchy::CreateParentItem(const FSceneOutl
 	{
 		if (UDataLayer* DataLayer = DataLayerActorTreeItem->GetDataLayer())
 		{
-			return Mode->CreateItemFor<FDataLayerTreeItem>(DataLayer, true);
+			return CreateDataLayerTreeItem(DataLayer, true);
 		}
 	}
 	else if (FDataLayerActorDescTreeItem* DataLayerActorDescTreeItem = Item->CastTo<FDataLayerActorDescTreeItem>())
 	{
 		if (UDataLayer* DataLayer = DataLayerActorDescTreeItem->GetDataLayer())
 		{
-			return Mode->CreateItemFor<FDataLayerTreeItem>(DataLayer, true);
+			return CreateDataLayerTreeItem(DataLayer, true);
 		}
 	}
 	return nullptr;
@@ -245,7 +279,7 @@ void FDataLayerHierarchy::OnWorldPartitionCreated(UWorld* InWorld)
 
 void FDataLayerHierarchy::OnLevelActorsAdded(const TArray<AActor*>& InActors)
 {
-	if (!bShowingDataLayerActors)
+	if (!bShowDataLayerActors)
 	{
 		return;
 	}
@@ -357,7 +391,7 @@ void FDataLayerHierarchy::OnLevelRemoved(ULevel* InLevel, UWorld* InWorld)
 
 void FDataLayerHierarchy::OnLoadedActorAdded(AActor& InActor)
 {
-	if (!bShowingDataLayerActors)
+	if (!bShowDataLayerActors)
 	{
 		return;
 	}
@@ -405,7 +439,7 @@ void FDataLayerHierarchy::OnLoadedActorRemoved(AActor& InActor)
 
 void FDataLayerHierarchy::OnActorDescAdded(FWorldPartitionActorDesc* InActorDesc)
 {
-	if (!bShowingUnloadedActors || !InActorDesc || InActorDesc->IsLoaded(true))
+	if (!bShowUnloadedActors || !InActorDesc || InActorDesc->IsLoaded(true))
 	{
 		return;
 	}
@@ -431,7 +465,7 @@ void FDataLayerHierarchy::OnActorDescAdded(FWorldPartitionActorDesc* InActorDesc
 
 void FDataLayerHierarchy::OnActorDescRemoved(FWorldPartitionActorDesc* InActorDesc)
 {
-	if (!bShowingUnloadedActors || (InActorDesc == nullptr))
+	if (!bShowUnloadedActors || (InActorDesc == nullptr))
 	{
 		return;
 	}
@@ -459,4 +493,4 @@ void FDataLayerHierarchy::FullRefreshEvent()
 	FSceneOutlinerHierarchyChangedData EventData;
 	EventData.Type = FSceneOutlinerHierarchyChangedData::FullRefresh;
 	HierarchyChangedEvent.Broadcast(EventData);
-}
\ No newline at end of file
+}
diff --git a/Engine/Source/Editor/DataLayerEditor/Private/DataLayer/DataLayerHierarchy.h b/Engine/Source/Editor/DataLayerEditor/Private/DataLayer/DataLayerHierarchy.h
--- a/Engine/Source/Editor/DataLayerEditor/Private/DataLayer/DataLayerHierarchy.h
+++ b/Engine/Source/Editor/DataLayerEditor/Private/DataLayer/DataLayerHierarchy.h
@@ -21,10 +21,12 @@ public:
 	virtual void CreateItems(TArray<FSceneOutlinerTreeItemPtr>& OutItems) const override;
 	virtual void CreateChildren(const FSceneOutlinerTreeItemPtr& Item, TArray<FSceneOutlinerTreeItemPtr>& OutChildren) const override {}
 	virtual FSceneOutlinerTreeItemPtr CreateParentItem(const FSceneOutlinerTreeItemPtr& Item) const override;
-	void SetShowingEditorDataLayers(bool bInShowingEditorDataLayers) { bShowingEditorDataLayers = bInShowingEditorDataLayers; }
-	void SetShowingRuntimeDataLayers(bool bInShowingRuntimeDataLayers) { bShowingRuntimeDataLayers = bInShowingRuntimeDataLayers; }
-	void SetShowingDataLayerActors(bool bInShowingDataLayerActors) { bShowingDataLayerActors = bInShowingDataLayerActors; }
-	void SetShowingUnloadedActors(bool bInShowingUnloadedActors) { bShowingUnloadedActors = bInShowingUnloadedActors; }
+	void SetShowEditorDataLayers(bool bInShowEditorDataLayers) { bShowEditorDataLayers = bInShowEditorDataLayers; }
+	void SetShowRuntimeDataLayers(bool bInShowRuntimeDataLayers) { bShowRuntimeDataLayers = bInShowRuntimeDataLayers; }
+	void SetShowDataLayerActors(bool bInShowDataLayerActors) { bShowDataLayerActors = bInShowDataLayerActors; }
+	void SetShowUnloadedActors(bool bInShowUnloadedActors) { bShowUnloadedActors = bInShowUnloadedActors; }
+	void SetShowOnlySelectedActors(bool bInbShowOnlySelectedActors) { bShowOnlySelectedActors = bInbShowOnlySelectedActors; }
+	void SetHighlightSelectedDataLayers(bool bInHighlightSelectedDataLayers) { bHighlightSelectedDataLayers = bInHighlightSelectedDataLayers; }
 
 private:
 	FDataLayerHierarchy(FDataLayerMode* Mode, const TWeakObjectPtr<UWorld>& Worlds);
@@ -46,10 +48,14 @@ private:
 	void OnActorDataLayersChanged(const TWeakObjectPtr<AActor>& InActor);
 	void OnDataLayerChanged(const EDataLayerAction Action, const TWeakObjectPtr<const UDataLayer>& ChangedDataLayer, const FName& ChangedProperty);
 	void FullRefreshEvent();
+	FSceneOutlinerTreeItemPtr CreateDataLayerTreeItem(UDataLayer* InDataLayer, bool bInForce = false) const;
+	bool IsDataLayerPartOfSelection(const UDataLayer* DataLayer) const;
 
 	TWeakObjectPtr<UWorld> RepresentingWorld;
-	bool bShowingEditorDataLayers;
-	bool bShowingRuntimeDataLayers;
-	bool bShowingDataLayerActors;
-	bool bShowingUnloadedActors;
+	bool bShowEditorDataLayers;
+	bool bShowRuntimeDataLayers;
+	bool bShowDataLayerActors;
+	bool bShowUnloadedActors;
+	bool bShowOnlySelectedActors;
+	bool bHighlightSelectedDataLayers;
 };
\ No newline at end of file
diff --git a/Engine/Source/Editor/DataLayerEditor/Private/DataLayer/DataLayerMode.cpp b/Engine/Source/Editor/DataLayerEditor/Private/DataLayer/DataLayerMode.cpp
--- a/Engine/Source/Editor/DataLayerEditor/Private/DataLayer/DataLayerMode.cpp
+++ b/Engine/Source/Editor/DataLayerEditor/Private/DataLayer/DataLayerMode.cpp
@@ -47,11 +47,32 @@ FDataLayerMode::FDataLayerMode(const FDataLayerModeParams& Params)
 	, SpecifiedWorldToDisplay(Params.SpecifiedWorldToDisplay)
 	, FilteredDataLayerCount(0)
 {
+	USelection::SelectionChangedEvent.AddRaw(this, &FDataLayerMode::OnLevelSelectionChanged);
+	USelection::SelectObjectEvent.AddRaw(this, &FDataLayerMode::OnLevelSelectionChanged);
+
 	UWorldPartitionEditorPerProjectUserSettings* SharedSettings = GetMutableDefault<UWorldPartitionEditorPerProjectUserSettings>();
 	bHideEditorDataLayers = SharedSettings->bHideEditorDataLayers;
 	bHideRuntimeDataLayers = SharedSettings->bHideRuntimeDataLayers;
 	bHideDataLayerActors = SharedSettings->bHideDataLayerActors;
 	bHideUnloadedActors = SharedSettings->bHideUnloadedActors;
+	bShowOnlySelectedActors = SharedSettings->bShowOnlySelectedActors;
+	bHighlightSelectedDataLayers = SharedSettings->bHighlightSelectedDataLayers;
+
+	FSceneOutlinerFilterInfo ShowOnlySelectedActorsInfo(LOCTEXT("ToggleShowOnlySelected", "Only Selected"), LOCTEXT("ToggleShowOnlySelectedToolTip", "When enabled, only displays actors that are currently selected."), bShowOnlySelectedActors, FCreateSceneOutlinerFilter::CreateStatic(&FDataLayerMode::CreateShowOnlySelectedActorsFilter));
+	ShowOnlySelectedActorsInfo.OnToggle().AddLambda([this](bool bIsActive)
+	{
+		UWorldPartitionEditorPerProjectUserSettings* Settings = GetMutableDefault<UWorldPartitionEditorPerProjectUserSettings>();
+		Settings->bShowOnlySelectedActors = bShowOnlySelectedActors = bIsActive;
+		Settings->PostEditChange();
+
+		if (auto DataLayerHierarchy = StaticCast<FDataLayerHierarchy*>(Hierarchy.Get()))
+		{
+			DataLayerHierarchy->SetShowOnlySelectedActors(bIsActive);
+		}
+
+		RefreshSelection();
+	});
+	FilterInfoMap.Add(TEXT("ShowOnlySelectedActors"), ShowOnlySelectedActorsInfo);
 
 	FSceneOutlinerFilterInfo HideEditorDataLayersInfo(LOCTEXT("ToggleHideEditorDataLayers", "Hide Editor Data Layers"), LOCTEXT("ToggleHideEditorDataLayersToolTip", "When enabled, hides Editor Data Layers."), bHideEditorDataLayers, FCreateSceneOutlinerFilter::CreateStatic(&FDataLayerMode::CreateHideEditorDataLayersFilter));
 	HideEditorDataLayersInfo.OnToggle().AddLambda([this](bool bIsActive)
@@ -62,7 +83,7 @@ FDataLayerMode::FDataLayerMode(const FDataLayerModeParams& Params)
 
 		if (auto DataLayerHierarchy = StaticCast<FDataLayerHierarchy*>(Hierarchy.Get()))
 		{
-			DataLayerHierarchy->SetShowingEditorDataLayers(!bIsActive);
+			DataLayerHierarchy->SetShowEditorDataLayers(!bIsActive);
 		}
 	});
 	FilterInfoMap.Add(TEXT("HideEditorDataLayersFilter"), HideEditorDataLayersInfo);
@@ -76,7 +97,7 @@ FDataLayerMode::FDataLayerMode(const FDataLayerModeParams& Params)
 
 		if (auto DataLayerHierarchy = StaticCast<FDataLayerHierarchy*>(Hierarchy.Get()))
 		{
-			DataLayerHierarchy->SetShowingRuntimeDataLayers(!bIsActive);
+			DataLayerHierarchy->SetShowRuntimeDataLayers(!bIsActive);
 		}
 	});
 	FilterInfoMap.Add(TEXT("HideRuntimeDataLayersFilter"), HideRuntimeDataLayersInfo);
@@ -90,7 +111,7 @@ FDataLayerMode::FDataLayerMode(const FDataLayerModeParams& Params)
 
 		if (auto DataLayerHierarchy = StaticCast<FDataLayerHierarchy*>(Hierarchy.Get()))
 		{
-			DataLayerHierarchy->SetShowingDataLayerActors(!bIsActive);
+			DataLayerHierarchy->SetShowDataLayerActors(!bIsActive);
 		}
 	});
 	FilterInfoMap.Add(TEXT("HideDataLayerActorsFilter"), HideDataLayerActorsInfo);
@@ -104,7 +125,7 @@ FDataLayerMode::FDataLayerMode(const FDataLayerModeParams& Params)
 
 		if (auto DataLayerHierarchy = StaticCast<FDataLayerHierarchy*>(Hierarchy.Get()))
 		{
-			DataLayerHierarchy->SetShowingUnloadedActors(!bIsActive);
+			DataLayerHierarchy->SetShowUnloadedActors(!bIsActive);
 		}
 	});
 	FilterInfoMap.Add(TEXT("HideUnloadedActorsFilter"), HideUnloadedActorsInfo);
@@ -119,6 +140,12 @@ FDataLayerMode::FDataLayerMode(const FDataLayerModeParams& Params)
 	const_cast<FSharedSceneOutlinerData&>(SceneOutliner->GetSharedData()).CustomDelete = FCustomSceneOutlinerDeleteDelegate::CreateRaw(this, &FDataLayerMode::DeleteItems);
 }
 
+FDataLayerMode::~FDataLayerMode()
+{
+	USelection::SelectionChangedEvent.RemoveAll(this);
+	USelection::SelectObjectEvent.RemoveAll(this);
+}
+
 TSharedRef<FSceneOutlinerFilter> FDataLayerMode::CreateHideEditorDataLayersFilter()
 {
 	return MakeShareable(new FDataLayerFilter(FDataLayerTreeItem::FFilterPredicate::CreateStatic([](const UDataLayer* DataLayer) { return true; }), FSceneOutlinerFilter::EDefaultBehaviour::Pass));
@@ -178,7 +205,7 @@ FText FDataLayerMode::GetStatusText() const
 	{
 		if (SelectedDataLayerCount == 0)
 		{
-			return FText::Format(LOCTEXT("ShowingAllDataLayersFmt", "{0} data layers"), FText::AsNumber(FilteredDataLayerCount), FText::AsNumber(FilteredDataLayerCount));
+			return FText::Format(LOCTEXT("ShowingAllDataLayersFmt", "{0} data layers"), FText::AsNumber(FilteredDataLayerCount));
 		}
 		else
 		{
@@ -206,12 +233,17 @@ SDataLayerBrowser* FDataLayerMode::GetDataLayerBrowser() const
 
 void FDataLayerMode::OnItemAdded(FSceneOutlinerTreeItemPtr Item)
 {
-	if (const FDataLayerTreeItem* DataLayerItem = Item->CastTo<FDataLayerTreeItem>())
+	if (FDataLayerTreeItem* DataLayerItem = Item->CastTo<FDataLayerTreeItem>())
 	{
 		if (!Item->Flags.bIsFilteredOut)
 		{
 			++FilteredDataLayerCount;
 
+			if (ShouldExpandDataLayer(DataLayerItem->GetDataLayer()))
+			{
+				SceneOutliner->SetItemExpansion(DataLayerItem->AsShared(), true);
+			}
+
 			if (SelectedDataLayersSet.Contains(DataLayerItem->GetDataLayer()))
 			{
 				SceneOutliner->AddToSelection({Item});
@@ -1022,6 +1054,32 @@ TSharedPtr<SWidget> FDataLayerMode::CreateContextMenu()
 
 void FDataLayerMode::CreateViewContent(FMenuBuilder& MenuBuilder)
 {
+	MenuBuilder.AddMenuEntry(
+		LOCTEXT("ToggleHighlightSelectedDataLayers", "Highlight Selected"),
+		LOCTEXT("ToggleHighlightSelectedDataLayersToolTip", "When enabled, highlights Data Layers containing actors that are currently selected."),
+		FSlateIcon(),
+		FUIAction(
+			FExecuteAction::CreateLambda([this]()
+			{
+				UWorldPartitionEditorPerProjectUserSettings* Settings = GetMutableDefault<UWorldPartitionEditorPerProjectUserSettings>();
+				bHighlightSelectedDataLayers = !bHighlightSelectedDataLayers;
+				Settings->bHighlightSelectedDataLayers = bHighlightSelectedDataLayers;
+				Settings->PostEditChange();
+
+				if (auto DataLayerHierarchy = StaticCast<FDataLayerHierarchy*>(Hierarchy.Get()))
+				{
+					DataLayerHierarchy->SetHighlightSelectedDataLayers(bHighlightSelectedDataLayers);
+				}
+
+				SceneOutliner->FullRefresh();
+			}), 
+			FCanExecuteAction(),
+			FIsActionChecked::CreateLambda([this]() { return bHighlightSelectedDataLayers; })
+		),
+		NAME_None,
+		EUserInterfaceActionType::ToggleButton
+	);
+
 	MenuBuilder.BeginSection("AssetThumbnails", LOCTEXT("ShowAdvancedHeading", "Advanced"));
 	MenuBuilder.AddMenuEntry(
 		LOCTEXT("ToggleAllowRuntimeDataLayerEditing", "Allow Runtime Data Layer Editing"), 
@@ -1130,10 +1188,12 @@ bool FDataLayerMode::IsWorldChecked(TWeakObjectPtr<UWorld> World) const
 TUniquePtr<ISceneOutlinerHierarchy> FDataLayerMode::CreateHierarchy()
 {
 	TUniquePtr<FDataLayerHierarchy> DataLayerHierarchy = FDataLayerHierarchy::Create(this, RepresentingWorld);
-	DataLayerHierarchy->SetShowingEditorDataLayers(!bHideEditorDataLayers);
-	DataLayerHierarchy->SetShowingRuntimeDataLayers(!bHideRuntimeDataLayers);
-	DataLayerHierarchy->SetShowingDataLayerActors(!bHideDataLayerActors);
-	DataLayerHierarchy->SetShowingUnloadedActors(!bHideUnloadedActors);
+	DataLayerHierarchy->SetShowEditorDataLayers(!bHideEditorDataLayers);
+	DataLayerHierarchy->SetShowRuntimeDataLayers(!bHideRuntimeDataLayers);
+	DataLayerHierarchy->SetShowDataLayerActors(!bHideDataLayerActors);
+	DataLayerHierarchy->SetShowUnloadedActors(!bHideUnloadedActors);
+	DataLayerHierarchy->SetShowOnlySelectedActors(bShowOnlySelectedActors);
+	DataLayerHierarchy->SetHighlightSelectedDataLayers(bHighlightSelectedDataLayers);
 	return DataLayerHierarchy;
 }
 
@@ -1247,6 +1307,100 @@ void FDataLayerMode::ChooseRepresentingWorld()
 	}
 }
 
+bool FDataLayerMode::ShouldExpandDataLayer(const UDataLayer* DataLayer) const
+{
+	if (bHighlightSelectedDataLayers || bShowOnlySelectedActors)
+	{
+		if (DataLayer)
+		{
+			if ((bShowOnlySelectedActors && DataLayerEditorSubsystem->DoesDataLayerContainSelectedActors(DataLayer)) ||
+				(ContainsSelectedChildDataLayer(DataLayer) && !DataLayer->GetChildren().IsEmpty()))
+			{
+				return true;
+			}
+		}
+	}
+	return false;
+}
+
+bool FDataLayerMode::ContainsSelectedChildDataLayer(const UDataLayer* DataLayer) const
+{
+	if (DataLayer)
+	{
+		bool bFoundSelected = false;
+		DataLayer->ForEachChild([this, &bFoundSelected](const UDataLayer* Child)
+		{
+			if (DataLayerEditorSubsystem->DoesDataLayerContainSelectedActors(Child) || ContainsSelectedChildDataLayer(Child))
+			{
+				bFoundSelected = true;
+				return false;
+			}
+			return true;
+		});
+		return bFoundSelected;
+	}
+	return false;
+};
+
+TSharedRef<FSceneOutlinerFilter> FDataLayerMode::CreateShowOnlySelectedActorsFilter()
+{
+	auto IsActorSelected = [](const AActor* InActor, const UDataLayer* InDataLayer)
+	{
+		return InActor && InActor->IsSelected();
+	};
+	return MakeShareable(new FDataLayerActorFilter(FDataLayerActorTreeItem::FFilterPredicate::CreateStatic(IsActorSelected), FSceneOutlinerFilter::EDefaultBehaviour::Pass, FDataLayerActorTreeItem::FFilterPredicate::CreateStatic(IsActorSelected)));
+}
+
+void FDataLayerMode::SynchronizeSelection()
+{
+	if (!bShowOnlySelectedActors && !bHighlightSelectedDataLayers)
+	{
+		return;
+	}
+
+	TArray<AActor*> Actors;
+	TSet<const UDataLayer*> ActorDataLayersIncludingParents;
+	GEditor->GetSelectedActors()->GetSelectedObjects<AActor>(Actors);
+	for (const AActor* Actor : Actors)
+	{
+		TArray<const UDataLayer*> ActorDataLayers = Actor->GetDataLayerObjects();
+		for (const UDataLayer* DataLayer : ActorDataLayers)
+		{
+			const UDataLayer* CurrentDataLayer = DataLayer;
+			while (CurrentDataLayer)
+			{
+				bool bIsAlreadyInSet = false;
+				ActorDataLayersIncludingParents.Add(CurrentDataLayer, &bIsAlreadyInSet);
+				if (!bIsAlreadyInSet)
+				{
+					FSceneOutlinerTreeItemPtr TreeItem = SceneOutliner->GetTreeItem(CurrentDataLayer, false);
+					if (TreeItem && ShouldExpandDataLayer(CurrentDataLayer))
+					{
+						SceneOutliner->SetItemExpansion(TreeItem, true);
+					}
+				}
+				CurrentDataLayer = CurrentDataLayer->GetParent();
+			}
+		}
+	}
+}
+
+void FDataLayerMode::OnLevelSelectionChanged(UObject* Obj)
+{
+	if (!bShowOnlySelectedActors && !bHighlightSelectedDataLayers)
+	{
+		return;
+	}
+
+	RefreshSelection();
+}
+
+void FDataLayerMode::RefreshSelection()
+{
+	SceneOutliner->FullRefresh();
+	SceneOutliner->RefreshSelection();
+}
+
 //
 // FDataLayerPickingMode : Lightweight version of FDataLayerMode used to show the DataLayer hierarchy and choose one.
 // 
@@ -1318,4 +1472,4 @@ void FDataLayerPickingMode::OnItemSelectionChanged(FSceneOutlinerTreeItemPtr Tre
 	}
 }
 
-#undef LOCTEXT_NAMESPACE
\ No newline at end of file
+#undef LOCTEXT_NAMESPACE
diff --git a/Engine/Source/Editor/DataLayerEditor/Private/DataLayer/DataLayerMode.h b/Engine/Source/Editor/DataLayerEditor/Private/DataLayer/DataLayerMode.h
--- a/Engine/Source/Editor/DataLayerEditor/Private/DataLayer/DataLayerMode.h
+++ b/Engine/Source/Editor/DataLayerEditor/Private/DataLayer/DataLayerMode.h
@@ -36,6 +36,7 @@ public:
 	};
 
 	FDataLayerMode(const FDataLayerModeParams& Params);
+	virtual ~FDataLayerMode();
 
 	virtual void Rebuild() override;
 	virtual TSharedPtr<SWidget> CreateContextMenu() override;
@@ -52,6 +53,7 @@ public:
 	virtual FText GetStatusText() const override;
 	virtual FSlateColor GetStatusTextColor() const override { return FSlateColor::UseForeground(); }
 
+	virtual void SynchronizeSelection() override;
 	virtual void OnItemDoubleClick(FSceneOutlinerTreeItemPtr Item) override;
 	virtual void OnItemAdded(FSceneOutlinerTreeItemPtr Item) override;
 	virtual void OnItemRemoved(FSceneOutlinerTreeItemPtr Item) override;
@@ -81,6 +83,10 @@ protected:
 	bool bHideDataLayerActors;
 	/** Should unloaded actors be hidden */
 	bool bHideUnloadedActors;
+	/** Should show only selected actors */
+	bool bShowOnlySelectedActors;
+	/** Should highlight DataLayers containing selected actors */
+	bool bHighlightSelectedDataLayers;
 	/** Delegate to call when an item is picked */
 	FOnSceneOutlinerItemPicked OnItemPicked;
 
@@ -94,9 +100,14 @@ private:
 	TArray<UDataLayer*> GetDataLayersFromOperation(const FDragDropOperation& Operation, bool bOnlyFindFirst = false) const;
 	TArray<UDataLayer*> GetSelectedDataLayers(SSceneOutliner* InSceneOutliner) const;
 	void SetParentDataLayer(const TArray<UDataLayer*> DataLayers, UDataLayer* ParentDataLayer) const;
+	void OnLevelSelectionChanged(UObject* Obj);
 	static void CreateDataLayerPicker(UToolMenu* InMenu, FOnDataLayerPicked OnDataLayerPicked, bool bInShowRoot = false);
+	bool ShouldExpandDataLayer(const UDataLayer* DataLayer) const;
+	bool ContainsSelectedChildDataLayer(const UDataLayer* DataLayer) const;
+	void RefreshSelection();
 
 	/** Filter factories */
+	static TSharedRef<FSceneOutlinerFilter> CreateShowOnlySelectedActorsFilter();
 	static TSharedRef<FSceneOutlinerFilter> CreateHideEditorDataLayersFilter();
 	static TSharedRef<FSceneOutlinerFilter> CreateHideRuntimeDataLayersFilter();
 	static TSharedRef<FSceneOutlinerFilter> CreateHideDataLayerActorsFilter();
@@ -136,6 +147,7 @@ public:
 	virtual void OnItemDoubleClick(FSceneOutlinerTreeItemPtr Item) override {}
 	virtual FReply OnKeyDown(const FKeyEvent& InKeyEvent) override { return FReply::Unhandled(); }
 	virtual void OnItemSelectionChanged(FSceneOutlinerTreeItemPtr TreeItem, ESelectInfo::Type SelectionType, const FSceneOutlinerItemSelection& Selection) override;
+	virtual void SynchronizeSelection() override {}
 
 	static TSharedRef<SWidget> CreateDataLayerPickerWidget(FOnDataLayerPicked OnDataLayerPicked);
 };
diff --git a/Engine/Source/Editor/DataLayerEditor/Private/DataLayer/DataLayerTreeItem.cpp b/Engine/Source/Editor/DataLayerEditor/Private/DataLayer/DataLayerTreeItem.cpp
--- a/Engine/Source/Editor/DataLayerEditor/Private/DataLayer/DataLayerTreeItem.cpp
+++ b/Engine/Source/Editor/DataLayerEditor/Private/DataLayer/DataLayerTreeItem.cpp
@@ -15,6 +15,7 @@ FDataLayerTreeItem::FDataLayerTreeItem(UDataLayer* InDataLayer)
 	: ISceneOutlinerTreeItem(Type)
 	, DataLayer(InDataLayer)
 	, ID(InDataLayer)
+	, bIsHighlighedtIfSelected(false)
 {
 	Flags.bIsExpanded = false;
 }
@@ -44,4 +45,16 @@ void FDataLayerTreeItem::OnVisibilityChanged(const bool bNewVisibility)
 	}
 }
 
+bool FDataLayerTreeItem::ShouldBeHighlighted() const
+{
+	if (bIsHighlighedtIfSelected)
+	{
+		if (UDataLayer* DataLayerPtr = DataLayer.Get())
+		{
+			return UDataLayerEditorSubsystem::Get()->DoesDataLayerContainSelectedActors(DataLayerPtr);
+		}
+	}
+	return false;
+}
+
 #undef LOCTEXT_NAMESPACE
\ No newline at end of file
diff --git a/Engine/Source/Editor/DataLayerEditor/Private/DataLayer/DataLayerTreeItem.h b/Engine/Source/Editor/DataLayerEditor/Private/DataLayer/DataLayerTreeItem.h
--- a/Engine/Source/Editor/DataLayerEditor/Private/DataLayer/DataLayerTreeItem.h
+++ b/Engine/Source/Editor/DataLayerEditor/Private/DataLayer/DataLayerTreeItem.h
@@ -26,6 +26,9 @@ public:
 	virtual bool ShouldShowVisibilityState() const { return true; }
 	/* End ISceneOutlinerTreeItem Implementation */
 
+	bool ShouldBeHighlighted() const;
+	void SetIsHighlightedIfSelected(bool bInIsHighlightedIfSelected) { bIsHighlighedtIfSelected = bInIsHighlightedIfSelected; }
+
 	static const FSceneOutlinerTreeItemType Type;
 
 	DECLARE_DELEGATE_RetVal_OneParam(bool, FFilterPredicate, const UDataLayer*);
@@ -44,4 +47,5 @@ public:
 private:
 	TWeakObjectPtr<UDataLayer> DataLayer;
 	const FObjectKey ID;
+	bool bIsHighlighedtIfSelected;
 };
\ No newline at end of file
diff --git a/Engine/Source/Editor/DataLayerEditor/Private/DataLayer/SDataLayerTreeLabel.cpp b/Engine/Source/Editor/DataLayerEditor/Private/DataLayer/SDataLayerTreeLabel.cpp
--- a/Engine/Source/Editor/DataLayerEditor/Private/DataLayer/SDataLayerTreeLabel.cpp
+++ b/Engine/Source/Editor/DataLayerEditor/Private/DataLayer/SDataLayerTreeLabel.cpp
@@ -28,6 +28,7 @@ void SDataLayerTreeLabel::Construct(const FArguments& InArgs, FDataLayerTreeItem
 	.VAlign(VAlign_Center)
 	[
 		SAssignNew(InlineTextBlock, SInlineEditableTextBlock)
+		.Font(this, &SDataLayerTreeLabel::GetDisplayNameFont)
 		.Text(this, &SDataLayerTreeLabel::GetDisplayText)
 		.ToolTipText(this, &SDataLayerTreeLabel::GetTooltipText)
 		.HighlightText(HighlightText)
@@ -97,6 +98,25 @@ void SDataLayerTreeLabel::Construct(const FArguments& InArgs, FDataLayerTreeItem
 	];
 }
 
+bool SDataLayerTreeLabel::ShouldBeHighlighted() const
+{
+	const FSceneOutlinerTreeItemPtr TreeItem = TreeItemPtr.Pin();
+	FDataLayerTreeItem* DataLayerTreeItem = TreeItem ? TreeItem->CastTo<FDataLayerTreeItem>() : nullptr;
+	return DataLayerTreeItem && DataLayerTreeItem->ShouldBeHighlighted();
+}
+
+FSlateFontInfo SDataLayerTreeLabel::GetDisplayNameFont() const
+{
+	if (ShouldBeHighlighted())
+	{
+		return FAppStyle::Get().GetFontStyle("DataLayerBrowser.LabelFontBold");
+	}
+	else
+	{
+		return FAppStyle::Get().GetFontStyle("DataLayerBrowser.LabelFont");
+	}
+}
+
 FText SDataLayerTreeLabel::GetDisplayText() const
 {
 	const UDataLayer* DataLayer = DataLayerPtr.Get();
@@ -193,7 +213,15 @@ FSlateColor SDataLayerTreeLabel::GetForegroundColor() const
 			return FSceneOutlinerCommonLabelData::DarkColor;
 		}
 	}
-	return (!DataLayer || !DataLayer->GetWorld()) ? FLinearColor(0.2f, 0.2f, 0.25f) : FSlateColor::UseForeground();
+	if (!DataLayer || !DataLayer->GetWorld())
+	{
+		return FLinearColor(0.2f, 0.2f, 0.25f);
+	}
+	if (ShouldBeHighlighted())
+	{
+		return FAppStyle::Get().GetSlateColor("Colors.AccentBlue");
+	}
+	return FSlateColor::UseForeground();
 }
 
 bool SDataLayerTreeLabel::OnVerifyItemLabelChanged(const FText& InLabel, FText& OutErrorMessage)
diff --git a/Engine/Source/Editor/DataLayerEditor/Private/DataLayer/SDataLayerTreeLabel.h b/Engine/Source/Editor/DataLayerEditor/Private/DataLayer/SDataLayerTreeLabel.h
--- a/Engine/Source/Editor/DataLayerEditor/Private/DataLayer/SDataLayerTreeLabel.h
+++ b/Engine/Source/Editor/DataLayerEditor/Private/DataLayer/SDataLayerTreeLabel.h
@@ -22,6 +22,7 @@ struct SDataLayerTreeLabel : FSceneOutlinerCommonLabelData, public SCompoundWidg
 
 private:
 
+	FSlateFontInfo GetDisplayNameFont() const;
 	FText GetDisplayText() const;
 	FText GetTooltipText() const;
 	FText GetTypeText() const;
@@ -31,6 +32,7 @@ private:
 	FSlateColor GetForegroundColor() const;
 	bool OnVerifyItemLabelChanged(const FText& InLabel, FText& OutErrorMessage);
 	void OnLabelCommitted(const FText& InLabel, ETextCommit::Type InCommitInfo);
+	bool ShouldBeHighlighted() const;
 
 	TWeakPtr<FDataLayerTreeItem> TreeItemPtr;
 	TWeakObjectPtr<UDataLayer> DataLayerPtr;
diff --git a/Engine/Source/Editor/DataLayerEditor/Public/DataLayer/DataLayerEditorSubsystem.h b/Engine/Source/Editor/DataLayerEditor/Public/DataLayer/DataLayerEditorSubsystem.h
--- a/Engine/Source/Editor/DataLayerEditor/Public/DataLayer/DataLayerEditorSubsystem.h
+++ b/Engine/Source/Editor/DataLayerEditor/Public/DataLayer/DataLayerEditorSubsystem.h
@@ -538,6 +538,11 @@ public:
 	*/
 	bool ResetUserSettings();
 
+	/*
+	 * Returns whether the DataLayer contains one of more actors that are part of the editor selection.
+	 */
+	bool DoesDataLayerContainSelectedActors(const UDataLayer* DataLayer) const { return SelectedDataLayersFromEditorSelection.Contains(DataLayer); }
+
 	//~ Begin Deprecated
 
 	UE_DEPRECATED(5.0, "Per-view Data Layer visibility was removed.")
@@ -626,6 +631,14 @@ private:
 	bool RefreshWorldPartitionEditorCells(bool bIsFromUserChange);
 	void UpdateDataLayerEditorPerProjectUserSettings();
 
+	void BroadcastActorDataLayersChanged(const TWeakObjectPtr<AActor>& ChangedActor);
+	void BroadcastDataLayerChanged(const EDataLayerAction Action, const TWeakObjectPtr<const UDataLayer>& ChangedDataLayer, const FName& ChangedProperty);
+	void OnSelectionChanged();
+	void RebuildSelectedDataLayersFromEditorSelection();
+
+	/** Contains Data Layers that contain actors that are part of the editor selection */
+	TSet<TWeakObjectPtr<const UDataLayer>> SelectedDataLayersFromEditorSelection;
+
 	/** Fires whenever one or more DataLayer changes */
 	FOnDataLayerChanged DataLayerChanged;
 
diff --git a/Engine/Source/Editor/EditorStyle/Private/StarshipStyle.cpp b/Engine/Source/Editor/EditorStyle/Private/StarshipStyle.cpp
--- a/Engine/Source/Editor/EditorStyle/Private/StarshipStyle.cpp
+++ b/Engine/Source/Editor/EditorStyle/Private/StarshipStyle.cpp
@@ -1176,6 +1176,8 @@ void FStarshipEditorStyle::FStyle::SetupGeneralStyles()
 			.SetHovered(BOX_BRUSH("Common/RoundedSelection_16x", 4.0f / 16.0f, SelectionColor))
 			.SetPressed(BOX_BRUSH("Common/RoundedSelection_16x", 4.0f / 16.0f, SelectionColor_Pressed))
 		);
+		Set("DataLayerBrowser.LabelFont", DEFAULT_FONT("Regular", 9));
+		Set("DataLayerBrowser.LabelFontBold", DEFAULT_FONT("Bold", 10));
 	}
 
 	// Derived Data
diff --git a/Engine/Source/Editor/SceneOutliner/Public/SceneOutlinerFilters.h b/Engine/Source/Editor/SceneOutliner/Public/SceneOutlinerFilters.h
--- a/Engine/Source/Editor/SceneOutliner/Public/SceneOutlinerFilters.h
+++ b/Engine/Source/Editor/SceneOutliner/Public/SceneOutlinerFilters.h
@@ -13,7 +13,7 @@ class FMenuBuilder;
 /**
 	* Contains information used to create a filter which will be displayed as user toggleable filter
 	*/
-class FSceneOutlinerFilterInfo
+class SCENEOUTLINER_API FSceneOutlinerFilterInfo
 { 
 public:
 	FSceneOutlinerFilterInfo(const FText& InFilterTitle, const FText& InFilterTooltip, bool bInActive, const FCreateSceneOutlinerFilter& InFactory = FCreateSceneOutlinerFilter())
diff --git a/Engine/Source/Runtime/Engine/Private/WorldPartition/DataLayer/DataLayer.cpp b/Engine/Source/Runtime/Engine/Private/WorldPartition/DataLayer/DataLayer.cpp
--- a/Engine/Source/Runtime/Engine/Private/WorldPartition/DataLayer/DataLayer.cpp
+++ b/Engine/Source/Runtime/Engine/Private/WorldPartition/DataLayer/DataLayer.cpp
@@ -193,6 +193,17 @@ void UDataLayer::SetParent(UDataLayer* InParent)
 	}
 }
 
+void UDataLayer::ForEachChild(TFunctionRef<bool(const UDataLayer*)> Operation) const
+{
+	for (UDataLayer* Child : Children)
+	{
+		if (!Operation(Child))
+		{
+			break;
+		}
+	}
+}
+
 void UDataLayer::SetChildParent(UDataLayer* InParent)
 {
 	if (this == InParent)
diff --git a/Engine/Source/Runtime/Engine/Public/WorldPartition/DataLayer/DataLayer.h b/Engine/Source/Runtime/Engine/Public/WorldPartition/DataLayer/DataLayer.h
--- a/Engine/Source/Runtime/Engine/Public/WorldPartition/DataLayer/DataLayer.h
+++ b/Engine/Source/Runtime/Engine/Public/WorldPartition/DataLayer/DataLayer.h
@@ -89,6 +89,7 @@ public:
 	bool CanParent(const UDataLayer* InParent) const;
 	void SetParent(UDataLayer* InParent);
 	void SetChildParent(UDataLayer* InParent);
+	void ForEachChild(TFunctionRef<bool(const UDataLayer*)> Operation) const;
 
 	static FText GetDataLayerText(const UDataLayer* InDataLayer);
 	const TCHAR* GetDataLayerIconName() const;
diff --git a/Engine/Source/Runtime/Engine/Public/WorldPartition/WorldPartitionEditorPerProjectUserSettings.h b/Engine/Source/Runtime/Engine/Public/WorldPartition/WorldPartitionEditorPerProjectUserSettings.h
--- a/Engine/Source/Runtime/Engine/Public/WorldPartition/WorldPartitionEditorPerProjectUserSettings.h
+++ b/Engine/Source/Runtime/Engine/Public/WorldPartition/WorldPartitionEditorPerProjectUserSettings.h
@@ -60,6 +60,8 @@ public:
 		, bHideDataLayerActors(true)
 		, bHideUnloadedActors(false)
 		, bAllowRuntimeDataLayerEditing(false)
+		, bShowOnlySelectedActors(false)
+		, bHighlightSelectedDataLayers(true)
 		, bDisableLoadingOfLastLoadedCells(false)
 #endif
 	{}
@@ -92,18 +94,26 @@ public:
 	UPROPERTY(config)
 	uint32 bHideRuntimeDataLayers : 1;
 
-	/** True when the DataLayer Outliner is not displaying actors */
+	/** True when the Data Layer Outliner is not displaying actors */
 	UPROPERTY(config)
 	uint32 bHideDataLayerActors : 1;
 
-	/** True when the DataLayer Outliner is not displaying unloaded actors */
+	/** True when the Data Layer Outliner is not displaying unloaded actors */
 	UPROPERTY(config)
 	uint32 bHideUnloadedActors : 1;
 
-	/** True when Runtime DataLayer editing is allowed. */
+	/** True when Runtime Data Layer editing is allowed. */
 	UPROPERTY(config)
 	uint32 bAllowRuntimeDataLayerEditing : 1;
 
+	/** True when the Data Layer Outliner is only displaying actors and datalayers for selected actors */
+	UPROPERTY(config)
+	uint32 bShowOnlySelectedActors : 1;
+
+	/** True when the Data Layer Outliner highlights Data Layers containing actors that are currently selected */
+	UPROPERTY(config)
+	uint32 bHighlightSelectedDataLayers : 1;
+
 private:
 	bool ShouldSaveSettings(const UWorld* InWorld) const
 	{
