Added buttons to quick add/remove selected actors to/from DataLayers.

diff --git a/Engine/Source/Editor/DataLayerEditor/Private/DataLayer/SDataLayerBrowser.cpp b/Engine/Source/Editor/DataLayerEditor/Private/DataLayer/SDataLayerBrowser.cpp
--- a/Engine/Source/Editor/DataLayerEditor/Private/DataLayer/SDataLayerBrowser.cpp
+++ b/Engine/Source/Editor/DataLayerEditor/Private/DataLayer/SDataLayerBrowser.cpp
@@ -92,7 +92,7 @@ void SDataLayerBrowser::Construct(const FArguments& InArgs)
 		// Data Layer Outliner
 		SNew(SSplitter)
 		.Orientation(Orient_Vertical)
-		.Style(FEditorStyle::Get(), "FoliageEditMode.Splitter")
+		.Style(FEditorStyle::Get(), "DetailsView.Splitter")
 		+ SSplitter::Slot()
 		[
 			SNew(SVerticalBox)
diff --git a/Engine/Source/Editor/DataLayerEditor/Private/DataLayer/SDataLayerOutliner.cpp b/Engine/Source/Editor/DataLayerEditor/Private/DataLayer/SDataLayerOutliner.cpp
new file mode 100644
--- /dev/null
+++ b/Engine/Source/Editor/DataLayerEditor/Private/DataLayer/SDataLayerOutliner.cpp
@@ -0,0 +1,101 @@
+// Copyright Epic Games, Inc. All Rights Reserved.
+
+#include "SDataLayerOutliner.h"
+#include "Framework/Application/SlateApplication.h"
+#include "Components/HorizontalBox.h"
+#include "Widgets/Input/SButton.h"
+#include "Widgets/Images/SImage.h"
+#include "DataLayer/DataLayerTreeItem.h"
+#include "DataLayer/DataLayerEditorSubsystem.h"
+#include "WorldPartition/DataLayer/DataLayer.h"
+#include "DataLayerTransaction.h"
+#include "Algo/Transform.h"
+#include "Editor.h"
+
+#define LOCTEXT_NAMESPACE "DataLayer"
+
+void SDataLayerOutliner::CustomAddToToolbar(TSharedPtr<SHorizontalBox> Toolbar)
+{
+	Toolbar->AddSlot()
+		.VAlign(VAlign_Center)
+		.AutoWidth()
+		.Padding(4.f, 0.f, 0.f, 0.f)
+		[
+			SNew(SButton)
+			.IsEnabled(this, &SDataLayerOutliner::CanAddSelectedActorsToSelectedDataLayersClicked)
+			.ButtonStyle(FAppStyle::Get(), "SimpleButton")
+			.ToolTipText(LOCTEXT("AddSelectedActorsToSelectedDataLayers", "Add selected actors to selected Data Layers"))
+			.OnClicked(this, &SDataLayerOutliner::OnAddSelectedActorsToSelectedDataLayersClicked)
+			[
+				SNew(SImage)
+				.ColorAndOpacity(FSlateColor::UseForeground())
+				.Image(FAppStyle::Get().GetBrush("DataLayerBrowser.AddSelection"))
+			]
+		];
+
+	Toolbar->AddSlot()
+		.VAlign(VAlign_Center)
+		.AutoWidth()
+		[
+			SNew(SButton)
+			.IsEnabled(this, &SDataLayerOutliner::CanRemoveSelectedActorsFromSelectedDataLayersClicked)
+			.ButtonStyle(FAppStyle::Get(), "SimpleButton")
+			.ToolTipText(LOCTEXT("RemoveSelectedActorsFromSelectedDataLayers", "Remove selected actors from selected Data Layers"))
+			.OnClicked(this, &SDataLayerOutliner::OnRemoveSelectedActorsFromSelectedDataLayersClicked)
+			[
+				SNew(SImage)
+				.ColorAndOpacity(FSlateColor::UseForeground())
+				.Image(FAppStyle::Get().GetBrush("DataLayerBrowser.RemoveSelection"))
+			]
+		];
+}
+
+TArray<UDataLayer*> SDataLayerOutliner::GetSelectedDataLayers() const
+{
+	FSceneOutlinerItemSelection ItemSelection(GetSelection());
+	TArray<FDataLayerTreeItem*> SelectedDataLayerItems;
+	ItemSelection.Get<FDataLayerTreeItem>(SelectedDataLayerItems);
+	TArray<UDataLayer*> ValidSelectedDataLayers;
+	Algo::TransformIf(SelectedDataLayerItems, ValidSelectedDataLayers, [](const auto Item) { return Item && Item->GetDataLayer(); }, [](const auto Item) { return Item->GetDataLayer(); });
+	return ValidSelectedDataLayers;
+}
+
+bool SDataLayerOutliner::CanAddSelectedActorsToSelectedDataLayersClicked() const
+{
+	if (GEditor->GetSelectedActorCount() > 0)
+	{
+		TArray<UDataLayer*> SelectedDataLayers = GetSelectedDataLayers();
+		const bool bSelectedDataLayersContainsLocked = !!SelectedDataLayers.FindByPredicate([](const UDataLayer* DataLayer) { return DataLayer->IsLocked(); });
+		return (!SelectedDataLayers.IsEmpty() && !bSelectedDataLayersContainsLocked);
+	}
+	return false;
+}
+
+bool SDataLayerOutliner::CanRemoveSelectedActorsFromSelectedDataLayersClicked() const
+{
+	return CanAddSelectedActorsToSelectedDataLayersClicked();
+}
+
+FReply SDataLayerOutliner::OnAddSelectedActorsToSelectedDataLayersClicked()
+{
+	if (CanAddSelectedActorsToSelectedDataLayersClicked())
+	{
+		TArray<UDataLayer*> SelectedDataLayers = GetSelectedDataLayers();
+		const FScopedDataLayerTransaction Transaction(LOCTEXT("AddSelectedActorsToSelectedDataLayers", "Add Selected Actors to Selected Data Layers"), SelectedDataLayers[0]->GetWorld());
+		UDataLayerEditorSubsystem::Get()->AddSelectedActorsToDataLayers(SelectedDataLayers);
+	}
+	return FReply::Handled();
+}
+
+FReply SDataLayerOutliner::OnRemoveSelectedActorsFromSelectedDataLayersClicked()
+{
+	if (CanRemoveSelectedActorsFromSelectedDataLayersClicked())
+	{
+		TArray<UDataLayer*> SelectedDataLayers = GetSelectedDataLayers();
+		const FScopedDataLayerTransaction Transaction(LOCTEXT("RemoveSelectedActorsFromSelectedDataLayers", "Remove Selected Actors from Selected Data Layers"), SelectedDataLayers[0]->GetWorld());
+		UDataLayerEditorSubsystem::Get()->RemoveSelectedActorsFromDataLayers(SelectedDataLayers);
+	}
+	return FReply::Handled();
+}
+
+#undef LOCTEXT_NAMESPACE
\ No newline at end of file
diff --git a/Engine/Source/Editor/DataLayerEditor/Private/DataLayer/SDataLayerOutliner.h b/Engine/Source/Editor/DataLayerEditor/Private/DataLayer/SDataLayerOutliner.h
--- a/Engine/Source/Editor/DataLayerEditor/Private/DataLayer/SDataLayerOutliner.h
+++ b/Engine/Source/Editor/DataLayerEditor/Private/DataLayer/SDataLayerOutliner.h
@@ -15,4 +15,12 @@ public:
 
 	SDataLayerOutliner() {}
 	virtual ~SDataLayerOutliner() {}
+
+	void CustomAddToToolbar(TSharedPtr<class SHorizontalBox> Toolbar) override;
+private:
+	bool CanAddSelectedActorsToSelectedDataLayersClicked() const;
+	bool CanRemoveSelectedActorsFromSelectedDataLayersClicked() const;
+	FReply OnAddSelectedActorsToSelectedDataLayersClicked();
+	FReply OnRemoveSelectedActorsFromSelectedDataLayersClicked();
+	TArray<class UDataLayer*> GetSelectedDataLayers() const;
 };
\ No newline at end of file
diff --git a/Engine/Source/Editor/EditorStyle/Private/StarshipStyle.cpp b/Engine/Source/Editor/EditorStyle/Private/StarshipStyle.cpp
--- a/Engine/Source/Editor/EditorStyle/Private/StarshipStyle.cpp
+++ b/Engine/Source/Editor/EditorStyle/Private/StarshipStyle.cpp
@@ -1169,6 +1169,8 @@ void FStarshipEditorStyle::FStyle::SetupGeneralStyles()
 		Set("DataLayer.Editor", new IMAGE_BRUSH_SVG("Icons/DataLayerEditor", Icon16x16));
 		Set("DataLayer.Runtime", new IMAGE_BRUSH_SVG("Icons/DataLayerRuntime", Icon16x16));
 		Set("DataLayer.LoadedInEditor", new IMAGE_BRUSH_SVG("Icons/DataLayerLoadedInEditor", Icon16x16));
+		Set("DataLayerBrowser.AddSelection", new IMAGE_BRUSH_SVG("Icons/DataLayerAddSelected", Icon16x16));
+		Set("DataLayerBrowser.RemoveSelection", new IMAGE_BRUSH_SVG("Icons/DataLayerRemoveSelected", Icon16x16));
 		Set("DataLayerBrowser.DataLayerContentsQuickbarBackground", new CORE_BOX_BRUSH("Common/DarkGroupBorder", 4.f / 16.f));
 		Set("DataLayerBrowser.Actor.RemoveFromDataLayer", new IMAGE_BRUSH("Icons/Cross_12x", Icon12x12));
 		Set("DataLayerBrowserButton", FButtonStyle(Button)
diff --git a/Engine/Source/Editor/SceneOutliner/Private/SSceneOutliner.cpp b/Engine/Source/Editor/SceneOutliner/Private/SSceneOutliner.cpp
--- a/Engine/Source/Editor/SceneOutliner/Private/SSceneOutliner.cpp
+++ b/Engine/Source/Editor/SceneOutliner/Private/SSceneOutliner.cpp
@@ -131,6 +131,7 @@ void SSceneOutliner::Construct(const FArguments& InArgs, const FSceneOutlinerIni
 		.OnTextCommitted( this, &SSceneOutliner::OnFilterTextCommitted )
 	];
 
+	CustomAddToToolbar(Toolbar);
 		
 	if (Mode->SupportsCreateNewFolder() && InInitOptions.bShowCreateNewFolder)
 	{
@@ -172,7 +173,6 @@ void SSceneOutliner::Construct(const FArguments& InArgs, const FSceneOutlinerIni
 		];
 	}
 
-
 	VerticalBox->AddSlot()
 	.AutoHeight()
 	.Padding( 8.0f, 8.0f, 8.0f, 4.0f )
diff --git a/Engine/Source/Editor/SceneOutliner/Public/SSceneOutliner.h b/Engine/Source/Editor/SceneOutliner/Public/SSceneOutliner.h
--- a/Engine/Source/Editor/SceneOutliner/Public/SSceneOutliner.h
+++ b/Engine/Source/Editor/SceneOutliner/Public/SSceneOutliner.h
@@ -349,6 +349,9 @@ public:
 	/** Tells the scene outliner that it should do a full refresh, which will clear the entire tree and rebuild it from scratch. */
 	virtual void FullRefresh() override;
 
+	/** Hook to add custom options to toolbar in a derived class */
+	virtual void CustomAddToToolbar(TSharedPtr<class SHorizontalBox> Toolbar) {}
+
 public:
 	/** Event to react to a user double click on a item */
 	SceneOutliner::FTreeItemPtrEvent& GetDoubleClickEvent() { return OnDoubleClickOnTreeEvent; }
