Added column Inital Runtime State in DataLayer Outliner.

diff --git a/Engine/Source/Editor/DataLayerEditor/Private/DataLayer/SDataLayerBrowser.cpp b/Engine/Source/Editor/DataLayerEditor/Private/DataLayer/SDataLayerBrowser.cpp
--- a/Engine/Source/Editor/DataLayerEditor/Private/DataLayer/SDataLayerBrowser.cpp
+++ b/Engine/Source/Editor/DataLayerEditor/Private/DataLayer/SDataLayerBrowser.cpp
@@ -72,6 +72,21 @@ void SDataLayerBrowser::Construct(const FArguments& InArgs)
 		return FString();
 	});
 
+	FGetTextForItem InternalInitialRuntimeStateInfoText = FGetTextForItem::CreateLambda([](const ISceneOutlinerTreeItem& Item) -> FString
+	{
+		if (const FDataLayerTreeItem* DataLayerItem = Item.CastTo<FDataLayerTreeItem>())
+		{
+			if (const UDataLayer* DataLayer = DataLayerItem->GetDataLayer())
+			{
+				if (DataLayer->IsRuntime())
+				{
+					return GetDataLayerRuntimeStateName(DataLayer->GetInitialRuntimeState());
+				}
+			}
+		}
+		return FString();
+	});
+
 	FSceneOutlinerInitializationOptions InitOptions;
 	InitOptions.bShowHeaderRow = true;
 	InitOptions.bShowParentTree = true;
@@ -81,7 +96,8 @@ void SDataLayerBrowser::Construct(const FArguments& InArgs)
 	InitOptions.ColumnMap.Add(FDataLayerOutlinerIsLoadedInEditorColumn::GetID(), FSceneOutlinerColumnInfo(ESceneOutlinerColumnVisibility::Visible, 1, FCreateSceneOutlinerColumn::CreateLambda([](ISceneOutliner& InSceneOutliner) { return MakeShareable(new FDataLayerOutlinerIsLoadedInEditorColumn(InSceneOutliner)); })));
 	InitOptions.ColumnMap.Add(FSceneOutlinerBuiltInColumnTypes::Label(), FSceneOutlinerColumnInfo(ESceneOutlinerColumnVisibility::Visible, 2));
 	InitOptions.ColumnMap.Add(FDataLayerOutlinerDeleteButtonColumn::GetID(), FSceneOutlinerColumnInfo(ESceneOutlinerColumnVisibility::Visible, 10, FCreateSceneOutlinerColumn::CreateLambda([](ISceneOutliner& InSceneOutliner) { return MakeShareable(new FDataLayerOutlinerDeleteButtonColumn(InSceneOutliner)); })));
-	InitOptions.ColumnMap.Add("ID Name", FSceneOutlinerColumnInfo(ESceneOutlinerColumnVisibility::Invisible, 20, FCreateSceneOutlinerColumn::CreateStatic(&FTextInfoColumn::CreateTextInfoColumn, FName("ID Name"), InternalNameInfoText)));
+	InitOptions.ColumnMap.Add("ID Name", FSceneOutlinerColumnInfo(ESceneOutlinerColumnVisibility::Invisible, 20, FCreateSceneOutlinerColumn::CreateStatic(&FTextInfoColumn::CreateTextInfoColumn, FName("ID Name"), InternalNameInfoText, FText::GetEmpty())));
+	InitOptions.ColumnMap.Add("Initial State", FSceneOutlinerColumnInfo(ESceneOutlinerColumnVisibility::Invisible, 20, FCreateSceneOutlinerColumn::CreateStatic(&FTextInfoColumn::CreateTextInfoColumn, FName("Initial State"), InternalInitialRuntimeStateInfoText, FText::FromString("Initial Runtime State"))));
 	DataLayerOutliner = SNew(SDataLayerOutliner, InitOptions).IsEnabled(FSlateApplication::Get().GetNormalExecutionAttribute());
 
 	SAssignNew(DataLayerContentsSection, SBorder)
diff --git a/Engine/Source/Editor/SceneOutliner/Private/SceneOutlinerModule.cpp b/Engine/Source/Editor/SceneOutliner/Private/SceneOutlinerModule.cpp
--- a/Engine/Source/Editor/SceneOutliner/Private/SceneOutlinerModule.cpp
+++ b/Engine/Source/Editor/SceneOutliner/Private/SceneOutlinerModule.cpp
@@ -372,7 +372,8 @@ void FSceneOutlinerModule::CreateActorInfoColumns(FSceneOutlinerInitializationOp
 				FCreateSceneOutlinerColumn::CreateStatic(
 					&FTextInfoColumn::CreateTextInfoColumn,
 					ColumnID,
-					ColumnInfo),
+					ColumnInfo,
+					FText::GetEmpty()),
 				true,
 				TOptional<float>(),
 				ColumnName));
diff --git a/Engine/Source/Editor/SceneOutliner/Private/SceneOutlinerTextInfoColumn.cpp b/Engine/Source/Editor/SceneOutliner/Private/SceneOutlinerTextInfoColumn.cpp
--- a/Engine/Source/Editor/SceneOutliner/Private/SceneOutlinerTextInfoColumn.cpp
+++ b/Engine/Source/Editor/SceneOutliner/Private/SceneOutlinerTextInfoColumn.cpp
@@ -9,13 +9,14 @@
 #define LOCTEXT_NAMESPACE "SceneOutlinerTextInfoColumn"
 
 
-TSharedRef<ISceneOutlinerColumn> FTextInfoColumn::CreateTextInfoColumn(ISceneOutliner& Outliner, const FName InColumnName, const FGetTextForItem InGetTextForItem)
+TSharedRef<ISceneOutlinerColumn> FTextInfoColumn::CreateTextInfoColumn(ISceneOutliner& Outliner, const FName InColumnName, const FGetTextForItem InGetTextForItem, const FText InColumnToolTip)
 {
-	return TSharedRef<ISceneOutlinerColumn>(MakeShareable(new FTextInfoColumn(Outliner, InColumnName, InGetTextForItem)));
+	return TSharedRef<ISceneOutlinerColumn>(MakeShareable(new FTextInfoColumn(Outliner, InColumnName, InGetTextForItem, InColumnToolTip)));
 }
 
-FTextInfoColumn::FTextInfoColumn(ISceneOutliner& Outliner, const FName InColumnName, const FGetTextForItem& InGetTextForItem):
+FTextInfoColumn::FTextInfoColumn(ISceneOutliner& Outliner, const FName InColumnName, const FGetTextForItem& InGetTextForItem, const FText InColumnToolTip):
 	ColumnName(InColumnName),
+	ColumnToolTip(InColumnToolTip),
 	SceneOutlinerWeak(StaticCastSharedRef<ISceneOutliner>(Outliner.AsShared())),
 	GetTextForItem(InGetTextForItem)
 {
@@ -30,6 +31,7 @@ FName FTextInfoColumn::GetColumnID()
 SHeaderRow::FColumn::FArguments FTextInfoColumn::ConstructHeaderRowColumn()
 {
 	return SHeaderRow::Column(ColumnName)
+		.DefaultTooltip(!ColumnToolTip.IsEmptyOrWhitespace() ? ColumnToolTip : TAttribute<FText>())
 		.FillWidth(2)
 		.HeaderComboVisibility(EHeaderComboVisibility::OnHover);
 }
diff --git a/Engine/Source/Editor/SceneOutliner/Public/SceneOutlinerTextInfoColumn.h b/Engine/Source/Editor/SceneOutliner/Public/SceneOutlinerTextInfoColumn.h
--- a/Engine/Source/Editor/SceneOutliner/Public/SceneOutlinerTextInfoColumn.h
+++ b/Engine/Source/Editor/SceneOutliner/Public/SceneOutlinerTextInfoColumn.h
@@ -30,12 +30,12 @@ class SCENEOUTLINER_API FTextInfoColumn : public ISceneOutlinerColumn
 		/**
 		 *	Constructor
 		 */
-		FTextInfoColumn(ISceneOutliner& Outliner, const FName InColumnName, const FGetTextForItem& InGetTextForItem);
+		FTextInfoColumn(ISceneOutliner& Outliner, const FName InColumnName, const FGetTextForItem& InGetTextForItem, const FText InColumnToolTip);
 
 		virtual ~FTextInfoColumn() {}
 
 		// Factory function to create a TextInfoColumn (This column is not registered by default because of it's requirements and therefore requires a factory!)
-		static TSharedRef<ISceneOutlinerColumn> CreateTextInfoColumn(ISceneOutliner& Outliner, const FName InColumnName, const FGetTextForItem InGetTextForItem);
+		static TSharedRef<ISceneOutlinerColumn> CreateTextInfoColumn(ISceneOutliner& Outliner, const FName InColumnName, const FGetTextForItem InGetTextForItem, const FText InColumnToolTip);
 
 		//////////////////////////////////////////////////////////////////////////
 		// Begin ISceneOutlinerColumn Implementation
@@ -61,6 +61,8 @@ class SCENEOUTLINER_API FTextInfoColumn : public ISceneOutlinerColumn
 
 		FName ColumnName;
 
+		FText ColumnToolTip;
+
 		/** Weak reference to the outliner widget that owns our list */
 		TWeakPtr< ISceneOutliner > SceneOutlinerWeak;
 
