Moved the "Allow Runtime Data Layer Editing" flag from the EditorPerProjectUserSettings to the AWorrldDataLayers actor so that the state persists for all users.

diff --git a/Engine/Source/Editor/DataLayerEditor/Private/DataLayer/DataLayerMode.cpp b/Engine/Source/Editor/DataLayerEditor/Private/DataLayer/DataLayerMode.cpp
--- a/Engine/Source/Editor/DataLayerEditor/Private/DataLayer/DataLayerMode.cpp
+++ b/Engine/Source/Editor/DataLayerEditor/Private/DataLayer/DataLayerMode.cpp
@@ -1088,13 +1088,19 @@ void FDataLayerMode::CreateViewContent(FMenuBuilder& MenuBuilder)
 		FUIAction(
 			FExecuteAction::CreateLambda([this]()
 			{
-				UWorldPartitionEditorPerProjectUserSettings* Settings = GetMutableDefault<UWorldPartitionEditorPerProjectUserSettings>();
-				Settings->bAllowRuntimeDataLayerEditing = !Settings->bAllowRuntimeDataLayerEditing;
-				Settings->PostEditChange();
+				if (AWorldDataLayers* WorldDataLayers = RepresentingWorld.IsValid() ? RepresentingWorld->GetWorldDataLayers() : nullptr)
+				{
+					const FScopedDataLayerTransaction Transaction(LOCTEXT("ToggleAllowRuntimeDataLayerEditingTransaction", "Toggle Allow Runtime Data Layer Editing"), RepresentingWorld.Get());
+					WorldDataLayers->SetAllowRuntimeDataLayerEditing(!WorldDataLayers->GetAllowRuntimeDataLayerEditing());
+				}
 				SceneOutliner->FullRefresh();
 			}),
 			FCanExecuteAction(),
-			FIsActionChecked::CreateLambda([]() { return !!GetMutableDefault<UWorldPartitionEditorPerProjectUserSettings>()->bAllowRuntimeDataLayerEditing; })
+			FIsActionChecked::CreateLambda([this]()
+			{
+				const AWorldDataLayers* WorldDataLayers = RepresentingWorld.IsValid() ? RepresentingWorld->GetWorldDataLayers() : nullptr;
+				return WorldDataLayers ? WorldDataLayers->GetAllowRuntimeDataLayerEditing() : true;
+			})
 		),
 		NAME_None,
 		EUserInterfaceActionType::ToggleButton
diff --git a/Engine/Source/Runtime/Engine/Private/WorldPartition/DataLayer/DataLayer.cpp b/Engine/Source/Runtime/Engine/Private/WorldPartition/DataLayer/DataLayer.cpp
--- a/Engine/Source/Runtime/Engine/Private/WorldPartition/DataLayer/DataLayer.cpp
+++ b/Engine/Source/Runtime/Engine/Private/WorldPartition/DataLayer/DataLayer.cpp
@@ -6,7 +6,6 @@
 
 #include "WorldPartition/DataLayer/DataLayer.h"
 #include "WorldPartition/DataLayer/WorldDataLayers.h"
-#include "WorldPartition/WorldPartitionEditorPerProjectUserSettings.h"
 
 #define LOCTEXT_NAMESPACE "DataLayer"
 
@@ -133,7 +132,7 @@ bool UDataLayer::IsLocked() const
 		return true;
 	}
 
-	return IsRuntime() && !GetMutableDefault<UWorldPartitionEditorPerProjectUserSettings>()->bAllowRuntimeDataLayerEditing;
+	return IsRuntime() && !GetOuterAWorldDataLayers()->GetAllowRuntimeDataLayerEditing();
 }
 
 bool UDataLayer::CanEditChange(const FProperty* InProperty) const
@@ -151,7 +150,7 @@ bool UDataLayer::CanEditChange(const FProperty* InProperty) const
 				return false;
 			}
 		}
-		return GetMutableDefault<UWorldPartitionEditorPerProjectUserSettings>()->bAllowRuntimeDataLayerEditing;
+		return GetOuterAWorldDataLayers()->GetAllowRuntimeDataLayerEditing();
 	}
 
 	return Super::CanEditChange(InProperty);
diff --git a/Engine/Source/Runtime/Engine/Private/WorldPartition/DataLayer/WorldDataLayers.cpp b/Engine/Source/Runtime/Engine/Private/WorldPartition/DataLayer/WorldDataLayers.cpp
--- a/Engine/Source/Runtime/Engine/Private/WorldPartition/DataLayer/WorldDataLayers.cpp
+++ b/Engine/Source/Runtime/Engine/Private/WorldPartition/DataLayer/WorldDataLayers.cpp
@@ -35,6 +35,9 @@ FString JoinDataLayerLabelsFromNames(AWorldDataLayers* InWorldDataLayers, const
 
 AWorldDataLayers::AWorldDataLayers(const FObjectInitializer& ObjectInitializer)
 	: Super(ObjectInitializer)
+#if WITH_EDITORONLY_DATA
+	, bAllowRuntimeDataLayerEditing(true)
+#endif
 {
 	bAlwaysRelevant = true;
 	bReplicates = true;
@@ -601,6 +604,14 @@ bool AWorldDataLayers::RemoveDataLayer(UDataLayer* InDataLayer)
 	return false;
 }
 
+void AWorldDataLayers::SetAllowRuntimeDataLayerEditing(bool bInAllowRuntimeDataLayerEditing)
+{
+	if (bAllowRuntimeDataLayerEditing != bInAllowRuntimeDataLayerEditing)
+	{
+		Modify();
+		bAllowRuntimeDataLayerEditing = bInAllowRuntimeDataLayerEditing;
+	}
+}
 #endif
 
 bool AWorldDataLayers::ContainsDataLayer(const UDataLayer* InDataLayer) const
diff --git a/Engine/Source/Runtime/Engine/Public/WorldPartition/DataLayer/WorldDataLayers.h b/Engine/Source/Runtime/Engine/Public/WorldPartition/DataLayer/WorldDataLayers.h
--- a/Engine/Source/Runtime/Engine/Public/WorldPartition/DataLayer/WorldDataLayers.h
+++ b/Engine/Source/Runtime/Engine/Public/WorldPartition/DataLayer/WorldDataLayers.h
@@ -35,6 +35,8 @@ public:
 	bool RemoveDataLayer(UDataLayer* InDataLayer);
 	bool RemoveDataLayers(const TArray<UDataLayer*>& InDataLayers);
 	FName GenerateUniqueDataLayerLabel(const FName& InDataLayerLabel) const;
+	void SetAllowRuntimeDataLayerEditing(bool bInAllowRuntimeDataLayerEditing);
+	bool GetAllowRuntimeDataLayerEditing() const { return bAllowRuntimeDataLayerEditing; }
 
 	//~ Begin Helper Functions
 	TArray<const UDataLayer*> GetDataLayerObjects(const TArray<FActorDataLayer>& DataLayers) const;
@@ -110,6 +112,12 @@ private:
 	TMap<FName, const UDataLayer*> NameToDataLayer;
 #endif
 
+#if WITH_EDITORONLY_DATA
+	// True when Runtime Data Layer editing is allowed.
+	UPROPERTY()
+	bool bAllowRuntimeDataLayerEditing;
+#endif
+
 	UPROPERTY()
 	TSet<TObjectPtr<UDataLayer>> WorldDataLayers;
 
diff --git a/Engine/Source/Runtime/Engine/Public/WorldPartition/WorldPartitionEditorPerProjectUserSettings.h b/Engine/Source/Runtime/Engine/Public/WorldPartition/WorldPartitionEditorPerProjectUserSettings.h
--- a/Engine/Source/Runtime/Engine/Public/WorldPartition/WorldPartitionEditorPerProjectUserSettings.h
+++ b/Engine/Source/Runtime/Engine/Public/WorldPartition/WorldPartitionEditorPerProjectUserSettings.h
@@ -59,7 +59,6 @@ public:
 		, bHideRuntimeDataLayers(false)
 		, bHideDataLayerActors(true)
 		, bHideUnloadedActors(false)
-		, bAllowRuntimeDataLayerEditing(false)
 		, bShowOnlySelectedActors(false)
 		, bHighlightSelectedDataLayers(true)
 		, bDisableLoadingOfLastLoadedCells(false)
@@ -102,10 +101,6 @@ public:
 	UPROPERTY(config)
 	uint32 bHideUnloadedActors : 1;
 
-	/** True when Runtime Data Layer editing is allowed. */
-	UPROPERTY(config)
-	uint32 bAllowRuntimeDataLayerEditing : 1;
-
 	/** True when the Data Layer Outliner is only displaying actors and datalayers for selected actors */
 	UPROPERTY(config)
 	uint32 bShowOnlySelectedActors : 1;
