Change unloaded data layers display color in PIE to be gray instead of red.

diff --git a/Engine/Source/Editor/DataLayerEditor/Private/DataLayer/SDataLayerTreeLabel.cpp b/Engine/Source/Editor/DataLayerEditor/Private/DataLayer/SDataLayerTreeLabel.cpp
--- a/Engine/Source/Editor/DataLayerEditor/Private/DataLayer/SDataLayerTreeLabel.cpp
+++ b/Engine/Source/Editor/DataLayerEditor/Private/DataLayer/SDataLayerTreeLabel.cpp
@@ -200,7 +200,7 @@ FSlateColor SDataLayerTreeLabel::GetForegroundColor() const
 				case EDataLayerRuntimeState::Loaded:
 					return FColorList::NeonBlue;
 				case EDataLayerRuntimeState::Unloaded:
-					return FColorList::Red;
+					return FColorList::DarkSlateGrey;
 				}
 			}
 			else
