Add feature to move actors from data layer(s) to another data layer in the Data Layers Outliner
* This is perform on Left Alt down.
* Mouse cursor is changed during the operation.
* SceneOutliner CreateDragDropOperation has now info on FPointerEvent.

diff --git a/Engine/Source/Editor/DataLayerEditor/Private/DataLayer/DataLayerDragDropOp.cpp b/Engine/Source/Editor/DataLayerEditor/Private/DataLayer/DataLayerDragDropOp.cpp
--- a/Engine/Source/Editor/DataLayerEditor/Private/DataLayer/DataLayerDragDropOp.cpp
+++ b/Engine/Source/Editor/DataLayerEditor/Private/DataLayer/DataLayerDragDropOp.cpp
@@ -1,7 +1,10 @@
 // Copyright Epic Games, Inc. All Rights Reserved.
 
 #include "DataLayer/DataLayerDragDropOp.h"
+#include "GameFramework/Actor.h"
 #include "Textures/SlateIcon.h"
+#include "ClassIconFinder.h"
+#include "Algo/Transform.h"
 
 void FDataLayerDragDropOp::Construct()
 {
@@ -19,3 +22,26 @@ void FDataLayerDragDropOp::Construct()
 	SetupDefaults();
 	FDecoratedDragDropOp::Construct();
 }
+
+void FDataLayerActorMoveOp::Construct()
+{
+	// Set text and icon
+	TArray< TWeakObjectPtr< AActor > > Actors;
+	Algo::TransformIf(DataLayerActorMoveElements, Actors, [](const FDataLayerActorMoveElement& MoveElement) { return MoveElement.Key.IsValid(); }, [](const FDataLayerActorMoveElement& MoveElement) { return MoveElement.Key; });
+
+	UClass* CommonSelClass = nullptr;
+	const FSlateBrush* Icon = FClassIconFinder::FindIconForActors(Actors, CommonSelClass);
+
+	if (DataLayerActorMoveElements.Num() == 1)
+	{
+		SetToolTip(FText::FromString(DataLayerActorMoveElements[0].Key->GetActorLabel()), Icon);
+	}
+	else
+	{
+		const FText Text = FText::Format(NSLOCTEXT("FDataLayerActorDragDropOp", "FormatActors", "{0} Actors"), DataLayerActorMoveElements.Num());
+		SetToolTip(Text, Icon);
+	}
+
+	SetupDefaults();
+	FDecoratedDragDropOp::Construct();
+}
\ No newline at end of file
diff --git a/Engine/Source/Editor/DataLayerEditor/Private/DataLayer/DataLayerDragDropOp.h b/Engine/Source/Editor/DataLayerEditor/Private/DataLayer/DataLayerDragDropOp.h
--- a/Engine/Source/Editor/DataLayerEditor/Private/DataLayer/DataLayerDragDropOp.h
+++ b/Engine/Source/Editor/DataLayerEditor/Private/DataLayer/DataLayerDragDropOp.h
@@ -5,6 +5,10 @@
 #include "DragAndDrop/DecoratedDragDropOp.h"
 
 class FDragDropEvent;
+class AActor;
+class UDataLayer;
+
+typedef TPair< TWeakObjectPtr<AActor>, TWeakObjectPtr<UDataLayer>> FDataLayerActorMoveElement;
 
 /** Drag/drop operation for dragging layers in the editor */
 class FDataLayerDragDropOp : public FDecoratedDragDropOp
@@ -17,3 +21,15 @@ public:
 
 	virtual void Construct() override;
 };
+
+/** Drag/drop operation for moving actor(s) in a data layer in the editor */
+class FDataLayerActorMoveOp : public FDecoratedDragDropOp
+{
+public:
+	DRAG_DROP_OPERATOR_TYPE(FDataLayerActorMoveOp, FDecoratedDragDropOp)
+
+	/** Actor that we are dragging */
+	TArray<FDataLayerActorMoveElement> DataLayerActorMoveElements;
+
+	virtual void Construct() override;
+};
\ No newline at end of file
diff --git a/Engine/Source/Editor/DataLayerEditor/Private/DataLayer/DataLayerMode.cpp b/Engine/Source/Editor/DataLayerEditor/Private/DataLayer/DataLayerMode.cpp
--- a/Engine/Source/Editor/DataLayerEditor/Private/DataLayer/DataLayerMode.cpp
+++ b/Engine/Source/Editor/DataLayerEditor/Private/DataLayer/DataLayerMode.cpp
@@ -12,7 +12,6 @@
 #include "DataLayerTreeItem.h"
 #include "WorldTreeItem.h"
 #include "WorldPartition/WorldPartitionEditorPerProjectUserSettings.h"
-#include "DataLayerDragDropOp.h"
 #include "ActorDescTreeItem.h"
 #include "ISceneOutlinerHierarchy.h"
 #include "SceneOutlinerMenuContext.h"
@@ -397,35 +396,27 @@ FReply FDataLayerMode::OnKeyDown(const FKeyEvent& InKeyEvent)
 
 bool FDataLayerMode::ParseDragDrop(FSceneOutlinerDragDropPayload& OutPayload, const FDragDropOperation& Operation) const
 {
-	return !GetActorsFromOperation(Operation, true).IsEmpty() || !GetDataLayersFromOperation(Operation, true).IsEmpty();
+	return !GetDataLayerActorPairsFromOperation(Operation).IsEmpty() || !GetActorsFromOperation(Operation, true).IsEmpty() || !GetDataLayersFromOperation(Operation, true).IsEmpty();
 }
 
-FSceneOutlinerDragValidationInfo FDataLayerMode::ValidateDrop(const ISceneOutlinerTreeItem& DropTarget, const FSceneOutlinerDragDropPayload& Payload) const
+FSceneOutlinerDragValidationInfo FDataLayerMode::ValidateDrop(const ISceneOutlinerTreeItem& DropTarget, bool bMoveOperation /*= false*/) const
 {
-	TArray<AActor*> PayloadActors = GetActorsFromOperation(Payload.SourceOperation);
-	if (!PayloadActors.IsEmpty())
+	if (const FDataLayerTreeItem* DataLayerItem = DropTarget.CastTo<FDataLayerTreeItem>())
 	{
-		for (AActor* Actor : PayloadActors)
+		const UDataLayer* TargetDataLayer = DataLayerItem->GetDataLayer();
+		if (!TargetDataLayer)
 		{
-			if (!DataLayerEditorSubsystem->IsActorValidForDataLayer(Actor))
-			{
-				return FSceneOutlinerDragValidationInfo(ESceneOutlinerDropCompatibility::IncompatibleGeneric, LOCTEXT("ActorCantBeAssignedToDataLayer", "Can't assign actors to Data Layer"));
-			}
+			return FSceneOutlinerDragValidationInfo(ESceneOutlinerDropCompatibility::IncompatibleGeneric, FText());
 		}
 
-		if (const FDataLayerTreeItem* DataLayerItem = DropTarget.CastTo<FDataLayerTreeItem>())
+		if (TargetDataLayer->IsLocked())
 		{
-			const UDataLayer* TargetDataLayer = DataLayerItem->GetDataLayer();
-			if (!TargetDataLayer)
-			{
-				return FSceneOutlinerDragValidationInfo(ESceneOutlinerDropCompatibility::IncompatibleGeneric, FText());
-			}
-			
-			if (TargetDataLayer->IsLocked())
-			{
-				return FSceneOutlinerDragValidationInfo(ESceneOutlinerDropCompatibility::IncompatibleGeneric, LOCTEXT("CantAssignLockedDataLayer", "Can't assign actors to locked Data Layer"));
-			}
+			const FText Text = bMoveOperation ? LOCTEXT("CantMoveLockedDataLayer", "Can't move actors to locked Data Layer") : LOCTEXT("CantAssignLockedDataLayer", "Can't assign actors to locked Data Layer");
+			return FSceneOutlinerDragValidationInfo(ESceneOutlinerDropCompatibility::IncompatibleGeneric, Text);
+		}
 
+		if (!bMoveOperation)
+		{
 			if (GetSelectedDataLayers(SceneOutliner).Num() > 1)
 			{
 				if (SceneOutliner->GetTree().IsItemSelected(const_cast<ISceneOutlinerTreeItem&>(DropTarget).AsShared()))
@@ -433,11 +424,35 @@ FSceneOutlinerDragValidationInfo FDataLayerMode::ValidateDrop(const ISceneOutlin
 					return FSceneOutlinerDragValidationInfo(ESceneOutlinerDropCompatibility::Compatible, LOCTEXT("AssignToDataLayers", "Assign to Selected Data Layers"));
 				}
 			}
+		}
+
+		const FText Text = bMoveOperation ? FText::Format(LOCTEXT("MoveToDataLayer", "Move to Data Layer \"{0}\""), FText::FromName(TargetDataLayer->GetDataLayerLabel())) : FText::Format(LOCTEXT("AssignToDataLayer", "Assign to Data Layer \"{0}\""), FText::FromName(TargetDataLayer->GetDataLayerLabel()));
+		return FSceneOutlinerDragValidationInfo(ESceneOutlinerDropCompatibility::Compatible, Text);
+	}
 
-			return FSceneOutlinerDragValidationInfo(ESceneOutlinerDropCompatibility::Compatible, FText::Format(LOCTEXT("AssignToDataLayer", "Assign to Data Layer \"{0}\""), FText::FromName(TargetDataLayer->GetDataLayerLabel())));
+	return FSceneOutlinerDragValidationInfo::Invalid();
+}
+
+FSceneOutlinerDragValidationInfo FDataLayerMode::ValidateDrop(const ISceneOutlinerTreeItem& DropTarget, const FSceneOutlinerDragDropPayload& Payload) const
+{
+	TArray<AActor*> PayloadActors = GetActorsFromOperation(Payload.SourceOperation);
+	if (!PayloadActors.IsEmpty()) // Adding actor(s) in data layer(s)
+	{
+		for (AActor* Actor : PayloadActors)
+		{
+			if (!DataLayerEditorSubsystem->IsActorValidForDataLayer(Actor))
+			{
+				return FSceneOutlinerDragValidationInfo(ESceneOutlinerDropCompatibility::IncompatibleGeneric, LOCTEXT("ActorCantBeAssignedToDataLayer", "Can't assign actors to Data Layer"));
+			}
 		}
+
+		return ValidateDrop(DropTarget);
+	}
+	else if (!GetDataLayerActorPairsFromOperation(Payload.SourceOperation).IsEmpty()) // Moving actor(s) into a Data Layer
+	{
+		return ValidateDrop(DropTarget, true);
 	}
-	else 
+	else // Moving a data layer(s)
 	{
 		TArray<UDataLayer*> PayloadDataLayers = GetDataLayersFromOperation(Payload.SourceOperation);
 		if (!PayloadDataLayers.IsEmpty())
@@ -487,7 +502,7 @@ FSceneOutlinerDragValidationInfo FDataLayerMode::ValidateDrop(const ISceneOutlin
 	return FSceneOutlinerDragValidationInfo::Invalid();
 }
 
-TArray<UDataLayer*> FDataLayerMode::GetDataLayersFromOperation(const FDragDropOperation& Operation, bool bOnlyFindFirst) const
+TArray<UDataLayer*> FDataLayerMode::GetDataLayersFromOperation(const FDragDropOperation& Operation, bool bOnlyFindFirst /*= false*/) const
 {
 	TArray<UDataLayer*> OutDataLayers;
 	
@@ -520,7 +535,20 @@ TArray<UDataLayer*> FDataLayerMode::GetDataLayersFromOperation(const FDragDropOp
 	return OutDataLayers;
 }
 
-TArray<AActor*> FDataLayerMode::GetActorsFromOperation(const FDragDropOperation& Operation, bool bOnlyFindFirst) const
+TArray<FDataLayerActorMoveElement> FDataLayerMode::GetDataLayerActorPairsFromOperation(const FDragDropOperation& Operation) const
+{
+	TArray<FDataLayerActorMoveElement> Out;
+
+	if (Operation.IsOfType<FDataLayerActorMoveOp>())
+	{
+		const FDataLayerActorMoveOp& DataLayerActorDragOp = StaticCast<const FDataLayerActorMoveOp&>(Operation);
+		return DataLayerActorDragOp.DataLayerActorMoveElements;
+	}
+
+	return Out;
+}
+
+TArray<AActor*> FDataLayerMode::GetActorsFromOperation(const FDragDropOperation& Operation, bool bOnlyFindFirst /*= false*/ ) const
 {
 	TSet<AActor*> Actors;
 
@@ -599,7 +627,7 @@ void FDataLayerMode::OnDrop(ISceneOutlinerTreeItem& DropTarget, const FSceneOutl
 	UDataLayer* TargetDataLayer = DataLayerItem ? DataLayerItem->GetDataLayer() : nullptr;
 
 	TArray<AActor*> ActorsToAdd = GetActorsFromOperation(Payload.SourceOperation);
-	if (!ActorsToAdd.IsEmpty())
+	if (!ActorsToAdd.IsEmpty()) // Adding actor(s) in data layer(s)
 	{
 		if (SceneOutliner->GetTree().IsItemSelected(const_cast<ISceneOutlinerTreeItem&>(DropTarget).AsShared()))
 		{
@@ -618,13 +646,52 @@ void FDataLayerMode::OnDrop(ISceneOutlinerTreeItem& DropTarget, const FSceneOutl
 			DataLayerEditorSubsystem->AddActorsToDataLayer(ActorsToAdd, TargetDataLayer);
 		}
 	}
-	else
+	else if (Payload.SourceOperation.IsOfType<FDataLayerActorMoveOp>()) // Moving actor(s) into a Data Layer
+	{
+		TArray<FDataLayerActorMoveElement> ActorsToMove = GetDataLayerActorPairsFromOperation(Payload.SourceOperation);
+		if (!ActorsToMove.IsEmpty() && TargetDataLayer)
+		{
+			const FScopedDataLayerTransaction Transaction(LOCTEXT("DataLayerOutlinerMoveActorsToDataLayer", "Move Actors to Data Layer"), RepresentingWorld.Get());
+			for (const auto& DataLayerActorPair : ActorsToMove)
+			{
+				if (AActor* ActorPtr = DataLayerActorPair.Key.Get())
+				{
+					DataLayerEditorSubsystem->AddActorToDataLayer(ActorPtr, TargetDataLayer);
+					DataLayerEditorSubsystem->RemoveActorFromDataLayer(ActorPtr, DataLayerActorPair.Value.Get());
+				}
+			}
+		}
+
+		return;
+	}
+	else // Moving a data layer(s)
 	{
 		TArray<UDataLayer*> DataLayers = GetDataLayersFromOperation(Payload.SourceOperation);
 		SetParentDataLayer(DataLayers, TargetDataLayer);
 	}
 }
 
+FReply FDataLayerMode::OnDragOverItem(const FDragDropEvent& Event, const ISceneOutlinerTreeItem& Item) const
+{
+	TSharedPtr<FDragDropOperation> DragOperation = Event.GetOperation();
+	
+	if (!DragOperation.IsValid())
+	{
+		return FReply::Handled();
+	}
+
+	if (DragOperation->IsOfType<FDataLayerActorMoveOp>()) // Moving actor(s) into a Data Layer
+	{
+		DragOperation->SetCursorOverride(EMouseCursor::GrabHandClosed);
+	}
+	else if (DragOperation->IsOfType<FActorDragDropOp>()) // Adding actor(s) in data layer(s)
+	{
+		DragOperation->SetCursorOverride(EMouseCursor::Default);
+	}
+
+	return FReply::Handled();
+}
+
 void FDataLayerMode::SetParentDataLayer(const TArray<UDataLayer*> DataLayers, UDataLayer* ParentDataLayer) const
 {
 	if (!DataLayers.IsEmpty())
@@ -669,6 +736,26 @@ struct FWeakDataLayerActorSelector
 	}
 };
 
+struct FDataLayerActorPairSelector
+{
+	bool operator()(const TWeakPtr<ISceneOutlinerTreeItem>& Item, FDataLayerActorMoveElement& DataOut) const
+	{
+		if (TSharedPtr<ISceneOutlinerTreeItem> ItemPtr = Item.Pin())
+		{
+			if (FDataLayerActorTreeItem* TypedItem = ItemPtr->CastTo<FDataLayerActorTreeItem>())
+			{
+				if (TypedItem->IsValid())
+				{
+					DataOut = FDataLayerActorMoveElement(TypedItem->Actor.Get(), TypedItem->GetDataLayer());
+					return true;
+				}
+			}
+		}
+		return false;
+	}
+};
+
+
 struct FWeakDataLayerSelector
 {
 	bool operator()(const TWeakPtr<ISceneOutlinerTreeItem>& Item, TWeakObjectPtr<UDataLayer>& DataOut) const
@@ -688,7 +775,7 @@ struct FWeakDataLayerSelector
 	}
 };
 
-TSharedPtr<FDragDropOperation> FDataLayerMode::CreateDragDropOperation(const TArray<FSceneOutlinerTreeItemPtr>& InTreeItems) const
+TSharedPtr<FDragDropOperation> FDataLayerMode::CreateDragDropOperation(const FPointerEvent& MouseEvent, const TArray<FSceneOutlinerTreeItemPtr>& InTreeItems) const
 {
 	FSceneOutlinerDragDropPayload DraggedObjects(InTreeItems);
 
@@ -724,13 +811,28 @@ TSharedPtr<FDragDropOperation> FDataLayerMode::CreateDragDropOperation(const TAr
 		return ActorOperation;
 	};
 
+	auto GetActorDataLayerOperation = [&DraggedObjects]()
+	{
+		TSharedPtr<FDataLayerActorMoveOp> DataLayerActorOperation = MakeShareable(new FDataLayerActorMoveOp);
+		DataLayerActorOperation->DataLayerActorMoveElements = DraggedObjects.GetData<FDataLayerActorMoveElement>(FDataLayerActorPairSelector());
+		DataLayerActorOperation->Construct();
+		return DataLayerActorOperation;
+	};
+
 	if (DraggedObjects.Has<FDataLayerTreeItem>() && !DraggedObjects.Has<FDataLayerActorTreeItem>())
 	{
 		return GetDataLayerOperation();
 	}
 	else if (!DraggedObjects.Has<FDataLayerTreeItem>() && DraggedObjects.Has<FDataLayerActorTreeItem>())
-	{
-		return GetActorOperation();
+	{	
+		if (MouseEvent.IsLeftAltDown())
+		{
+			return GetActorDataLayerOperation();
+		}
+		else
+		{
+			return GetActorOperation();
+		}
 	}
 	else
 	{
@@ -764,7 +866,7 @@ TArray<UDataLayer*> FDataLayerMode::GetSelectedDataLayers(SSceneOutliner* InScen
 	return MoveTemp(ValidSelectedDataLayers);
 }
 
-void FDataLayerMode::CreateDataLayerPicker(UToolMenu* InMenu, FOnDataLayerPicked OnDataLayerPicked, bool bInShowRoot)
+void FDataLayerMode::CreateDataLayerPicker(UToolMenu* InMenu, FOnDataLayerPicked OnDataLayerPicked, bool bInShowRoot /*= false*/)
 {
 	if (bInShowRoot)
 	{
diff --git a/Engine/Source/Editor/DataLayerEditor/Private/DataLayer/DataLayerMode.h b/Engine/Source/Editor/DataLayerEditor/Private/DataLayer/DataLayerMode.h
--- a/Engine/Source/Editor/DataLayerEditor/Private/DataLayer/DataLayerMode.h
+++ b/Engine/Source/Editor/DataLayerEditor/Private/DataLayer/DataLayerMode.h
@@ -6,6 +6,7 @@
 #include "SceneOutlinerFwd.h"
 #include "ISceneOutlinerMode.h"
 #include "DataLayerTreeItem.h"
+#include "DataLayerDragDropOp.h"
 
 class UDataLayerEditorSubsystem;
 class SDataLayerBrowser;
@@ -65,7 +66,8 @@ public:
 	virtual bool ParseDragDrop(FSceneOutlinerDragDropPayload& OutPayload, const FDragDropOperation& Operation) const override;
 	virtual FSceneOutlinerDragValidationInfo ValidateDrop(const ISceneOutlinerTreeItem& DropTarget, const FSceneOutlinerDragDropPayload& Payload) const override;
 	virtual void OnDrop(ISceneOutlinerTreeItem& DropTarget, const FSceneOutlinerDragDropPayload& Payload, const FSceneOutlinerDragValidationInfo& ValidationInfo) const override;
-	virtual TSharedPtr<FDragDropOperation> CreateDragDropOperation(const TArray<FSceneOutlinerTreeItemPtr>& InTreeItems) const override;
+	virtual TSharedPtr<FDragDropOperation> CreateDragDropOperation(const FPointerEvent& MouseEvent, const TArray<FSceneOutlinerTreeItemPtr>& InTreeItems) const override;
+	virtual FReply OnDragOverItem(const FDragDropEvent& Event, const ISceneOutlinerTreeItem& Item) const override; 
 
 	void DeleteItems(const TArray<TWeakPtr<ISceneOutlinerTreeItem>>& Items);
 	SDataLayerBrowser* GetDataLayerBrowser() const;
@@ -96,6 +98,7 @@ private:
 	void ChooseRepresentingWorld();
 	void OnSelectWorld(TWeakObjectPtr<UWorld> World);
 	bool IsWorldChecked(TWeakObjectPtr<UWorld> World) const;
+	TArray<FDataLayerActorMoveElement> GetDataLayerActorPairsFromOperation(const FDragDropOperation& Operation) const;
 	TArray<AActor*> GetActorsFromOperation(const FDragDropOperation& Operation, bool bOnlyFindFirst = false) const;
 	TArray<UDataLayer*> GetDataLayersFromOperation(const FDragDropOperation& Operation, bool bOnlyFindFirst = false) const;
 	TArray<UDataLayer*> GetSelectedDataLayers(SSceneOutliner* InSceneOutliner) const;
@@ -105,6 +108,7 @@ private:
 	bool ShouldExpandDataLayer(const UDataLayer* DataLayer) const;
 	bool ContainsSelectedChildDataLayer(const UDataLayer* DataLayer) const;
 	void RefreshSelection();
+	FSceneOutlinerDragValidationInfo ValidateDrop(const ISceneOutlinerTreeItem& DropTarget, bool bMoveOperation = false) const;
 
 	/** Filter factories */
 	static TSharedRef<FSceneOutlinerFilter> CreateShowOnlySelectedActorsFilter();
diff --git a/Engine/Source/Editor/SceneOutliner/Private/ActorBrowsingMode.cpp b/Engine/Source/Editor/SceneOutliner/Private/ActorBrowsingMode.cpp
--- a/Engine/Source/Editor/SceneOutliner/Private/ActorBrowsingMode.cpp
+++ b/Engine/Source/Editor/SceneOutliner/Private/ActorBrowsingMode.cpp
@@ -907,7 +907,7 @@ bool FActorBrowsingMode::GetFolderNamesFromPayload(const FSceneOutlinerDragDropP
 	return GetFolderNamesFromFolders(InPayload.GetData<FFolder>(SceneOutliner::FFolderPathSelector()), OutFolders, OutCommonRootObject);
 }
 
-TSharedPtr<FDragDropOperation> FActorBrowsingMode::CreateDragDropOperation(const TArray<FSceneOutlinerTreeItemPtr>& InTreeItems) const
+TSharedPtr<FDragDropOperation> FActorBrowsingMode::CreateDragDropOperation(const FPointerEvent& MouseEvent, const TArray<FSceneOutlinerTreeItemPtr>& InTreeItems) const
 {
 	FSceneOutlinerDragDropPayload DraggedObjects(InTreeItems);
 
diff --git a/Engine/Source/Editor/SceneOutliner/Private/SOutlinerTreeView.cpp b/Engine/Source/Editor/SceneOutliner/Private/SOutlinerTreeView.cpp
--- a/Engine/Source/Editor/SceneOutliner/Private/SOutlinerTreeView.cpp
+++ b/Engine/Source/Editor/SceneOutliner/Private/SOutlinerTreeView.cpp
@@ -47,7 +47,7 @@ static FReply HandleOnDragDetected( const FGeometry& MyGeometry, const FPointerE
 	auto TablePtr = Table.Pin();
 	if (TablePtr.IsValid() && MouseEvent.IsMouseButtonDown( EKeys::LeftMouseButton ))
 	{
-		auto Operation = TablePtr->GetOutlinerPtr().Pin()->CreateDragDropOperation(TablePtr->GetSelectedItems());
+		auto Operation = TablePtr->GetOutlinerPtr().Pin()->CreateDragDropOperation(MouseEvent, TablePtr->GetSelectedItems());
 
 		if (Operation.IsValid())
 		{
diff --git a/Engine/Source/Editor/SceneOutliner/Private/SSceneOutliner.cpp b/Engine/Source/Editor/SceneOutliner/Private/SSceneOutliner.cpp
--- a/Engine/Source/Editor/SceneOutliner/Private/SSceneOutliner.cpp
+++ b/Engine/Source/Editor/SceneOutliner/Private/SSceneOutliner.cpp
@@ -2227,9 +2227,9 @@ void SSceneOutliner::SetItemExpansionRecursive(FSceneOutlinerTreeItemPtr Model,
 	}
 }
 
-TSharedPtr<FDragDropOperation> SSceneOutliner::CreateDragDropOperation(const TArray<FSceneOutlinerTreeItemPtr>& InTreeItems) const 
+TSharedPtr<FDragDropOperation> SSceneOutliner::CreateDragDropOperation(const FPointerEvent& MouseEvent, const TArray<FSceneOutlinerTreeItemPtr>& InTreeItems) const
 { 
-	return Mode->CreateDragDropOperation(InTreeItems); 
+	return Mode->CreateDragDropOperation(MouseEvent, InTreeItems); 
 }
 
 /** Parse a drag drop operation into a payload */
diff --git a/Engine/Source/Editor/SceneOutliner/Public/ActorBrowsingMode.h b/Engine/Source/Editor/SceneOutliner/Public/ActorBrowsingMode.h
--- a/Engine/Source/Editor/SceneOutliner/Public/ActorBrowsingMode.h
+++ b/Engine/Source/Editor/SceneOutliner/Public/ActorBrowsingMode.h
@@ -16,7 +16,7 @@ public:
 	virtual void Rebuild() override;
 	virtual FCreateSceneOutlinerMode CreateFolderPickerMode(const FFolder::FRootObject& InRootObject = FFolder::GetDefaultRootObject()) const override;
 	virtual void CreateViewContent(FMenuBuilder& MenuBuilder) override;
-	virtual TSharedPtr<FDragDropOperation> CreateDragDropOperation(const TArray<FSceneOutlinerTreeItemPtr>& InTreeItems) const override;
+	virtual TSharedPtr<FDragDropOperation> CreateDragDropOperation(const FPointerEvent& MouseEvent, const TArray<FSceneOutlinerTreeItemPtr>& InTreeItems) const override;
 	virtual bool ParseDragDrop(FSceneOutlinerDragDropPayload& OutPayload, const FDragDropOperation& Operation) const override;
 	virtual FSceneOutlinerDragValidationInfo ValidateDrop(const ISceneOutlinerTreeItem& DropTarget, const FSceneOutlinerDragDropPayload& Payload) const override;
 	virtual TSharedPtr<SWidget> CreateContextMenu() override;
diff --git a/Engine/Source/Editor/SceneOutliner/Public/ISceneOutlinerMode.h b/Engine/Source/Editor/SceneOutliner/Public/ISceneOutlinerMode.h
--- a/Engine/Source/Editor/SceneOutliner/Public/ISceneOutlinerMode.h
+++ b/Engine/Source/Editor/SceneOutliner/Public/ISceneOutlinerMode.h
@@ -41,7 +41,7 @@ public:
 	virtual FCreateSceneOutlinerMode CreateFolderPickerMode(const FFolder::FRootObject& InRootObject = FFolder::GetDefaultRootObject()) const { return FCreateSceneOutlinerMode(); }
 public:
 	/** Construct a new Drag and drop operation */
-	virtual TSharedPtr<FDragDropOperation> CreateDragDropOperation(const TArray<FSceneOutlinerTreeItemPtr>& InTreeItems) const { return nullptr; }
+	virtual TSharedPtr<FDragDropOperation> CreateDragDropOperation(const FPointerEvent& MouseEvent, const TArray<FSceneOutlinerTreeItemPtr>& InTreeItems) const { return nullptr; }
 
 	/**
 		* Populates a drag/drop operation if the mode supports that type of operation.
diff --git a/Engine/Source/Editor/SceneOutliner/Public/SSceneOutliner.h b/Engine/Source/Editor/SceneOutliner/Public/SSceneOutliner.h
--- a/Engine/Source/Editor/SceneOutliner/Public/SSceneOutliner.h
+++ b/Engine/Source/Editor/SceneOutliner/Public/SSceneOutliner.h
@@ -451,7 +451,7 @@ public:
 	TSharedPtr<FSceneOutlinerFilters>& GetFilters() { return Filters; }
 
 	/** Create a drag drop operation */
-	TSharedPtr<FDragDropOperation> CreateDragDropOperation(const TArray<FSceneOutlinerTreeItemPtr>& InTreeItems) const;
+	TSharedPtr<FDragDropOperation> CreateDragDropOperation(const FPointerEvent& MouseEvent, const TArray<FSceneOutlinerTreeItemPtr>& InTreeItems) const;
 
 	/** Parse a drag drop operation into a payload */
 	bool ParseDragDrop(FSceneOutlinerDragDropPayload& OutPayload, const FDragDropOperation& Operation) const;
