Data Layer picking mode now hides custom toolbar used by default Data Layer mode.

diff --git a/Engine/Source/Editor/DataLayerEditor/Private/DataLayer/DataLayerMode.h b/Engine/Source/Editor/DataLayerEditor/Private/DataLayer/DataLayerMode.h
--- a/Engine/Source/Editor/DataLayerEditor/Private/DataLayer/DataLayerMode.h
+++ b/Engine/Source/Editor/DataLayerEditor/Private/DataLayer/DataLayerMode.h
@@ -48,6 +48,7 @@ public:
 	virtual bool IsInteractive() const override { return true; }
 	virtual bool CanRename() const override { return true; }
 	virtual bool CanRenameItem(const ISceneOutlinerTreeItem& Item) const override;
+	virtual bool CanCustomizeToolbar() const { return true; }
 	virtual bool ShowStatusBar() const override { return true; }
 	virtual bool ShowViewButton() const override { return true; }
 	virtual bool ShowFilterOptions() const override { return true; }
@@ -148,6 +149,7 @@ public:
 	virtual bool SupportsKeyboardFocus() const override { return false; }
 	virtual bool CanRename() const override { return false; }
 	virtual bool CanRenameItem(const ISceneOutlinerTreeItem& Item) const override { return false; }
+	virtual bool CanCustomizeToolbar() const { return false; }
 	virtual void OnItemDoubleClick(FSceneOutlinerTreeItemPtr Item) override {}
 	virtual FReply OnKeyDown(const FKeyEvent& InKeyEvent) override { return FReply::Unhandled(); }
 	virtual void OnItemSelectionChanged(FSceneOutlinerTreeItemPtr TreeItem, ESelectInfo::Type SelectionType, const FSceneOutlinerItemSelection& Selection) override;
diff --git a/Engine/Source/Editor/SceneOutliner/Private/SSceneOutliner.cpp b/Engine/Source/Editor/SceneOutliner/Private/SSceneOutliner.cpp
--- a/Engine/Source/Editor/SceneOutliner/Private/SSceneOutliner.cpp
+++ b/Engine/Source/Editor/SceneOutliner/Private/SSceneOutliner.cpp
@@ -145,7 +145,10 @@ void SSceneOutliner::Construct(const FArguments& InArgs, const FSceneOutlinerIni
 		.OnTextCommitted( this, &SSceneOutliner::OnFilterTextCommitted )
 	];
 
-	CustomAddToToolbar(Toolbar);
+	if (Mode->CanCustomizeToolbar())
+	{
+		CustomAddToToolbar(Toolbar);
+	}
 		
 	if (Mode->SupportsCreateNewFolder() && InInitOptions.bShowCreateNewFolder)
 	{
diff --git a/Engine/Source/Editor/SceneOutliner/Public/ISceneOutlinerMode.h b/Engine/Source/Editor/SceneOutliner/Public/ISceneOutlinerMode.h
--- a/Engine/Source/Editor/SceneOutliner/Public/ISceneOutlinerMode.h
+++ b/Engine/Source/Editor/SceneOutliner/Public/ISceneOutlinerMode.h
@@ -82,6 +82,8 @@ public:
 	virtual TSharedPtr<SWidget> CreateContextMenu() { return TSharedPtr<SWidget>(); }
 	/** Check if an item can be renamed */
 	virtual bool CanRenameItem(const ISceneOutlinerTreeItem& Item) const { return false; }
+	/** Whether the toolbar can be customized. */
+	virtual bool CanCustomizeToolbar() const { return false; }
 
 	/** Synchronize the mode specific selection with the tree view */
 	virtual void SynchronizeSelection() {}
