UDataLayer deprecated in favor of UDataLayerInstance & UDataLayerAsset
UDeprecatedDataLayerInstance allows to boot level using deprecated UDataLayers
Worlds using DataLayer can run the DataLayerToAssetCommandlet to transition their existing DataLayers
FActorDataLayer interface is deprecated.
DataLayers Blueprint referencers should now use DataLayerAsset to retrieve DataLayerInstances.
DataLayer Code referencers should now use DataLayerAssets or DataLayerInstance FName to retrieve DataLayerInstances.
DataLayerLabels now only used for display/UI purpose
Relabeling DataLayers is not permitted anymore on new DataLayerInstances (allowed on UDeprecatedDataLayerInstance)
Added Changelist Validation for DataLayers
Added a column to the datalayer outliner showing any data layer errors.

diff --git a/Engine/Plugins/Editor/DataValidation/Source/DataValidation/Private/WorldPartitionChangelistValidator.cpp b/Engine/Plugins/Editor/DataValidation/Source/DataValidation/Private/WorldPartitionChangelistValidator.cpp
--- a/Engine/Plugins/Editor/DataValidation/Source/DataValidation/Private/WorldPartitionChangelistValidator.cpp
+++ b/Engine/Plugins/Editor/DataValidation/Source/DataValidation/Private/WorldPartitionChangelistValidator.cpp
@@ -13,6 +13,8 @@
 #include "WorldPartition/WorldPartitionActorDesc.h"
 #include "WorldPartition/ActorDescContainer.h"
 #include "WorldPartition/WorldPartition.h"
+#include "WorldPartition/DataLayer/DataLayerInstanceWithAsset.h"
+#include "WorldPartition/DataLayer/WorldDataLayers.h"
 
 #define LOCTEXT_NAMESPACE "WorldPartitionChangelistValidation"
 
@@ -29,7 +31,7 @@ EDataValidationResult UWorldPartitionChangelistValidator::ValidateLoadedAsset_Im
 	
 	Errors = &ValidationErrors;
 	
-	EDataValidationResult Result = ValidateActorsListFromChangeList(ChangeList);
+	EDataValidationResult Result = ValidateActorsAndDataLayersFromChangeList(ChangeList);
 
 	if (Result == EDataValidationResult::Invalid)
 	{
@@ -66,14 +68,46 @@ FString GetPrettyPackageName(const FWorldPartitionActorDescView& Desc)
 // Extract all Actors/Map from Changelist (in OFPA this should be one Actor per Package, and we'll discard all Actors from non WorldPartition maps)
 // and add them to a Map of World->Files[] so that we can do one validation per world. Once Worlds are identified, we either the UActorDescContainer 
 // from memory (if loaded) or request it to be loaded, we then build a Set of objects that interest us from the Actors in the CL 
-EDataValidationResult UWorldPartitionChangelistValidator::ValidateActorsListFromChangeList(UDataValidationChangelist* Changelist)
+EDataValidationResult UWorldPartitionChangelistValidator::ValidateActorsAndDataLayersFromChangeList(UDataValidationChangelist* Changelist)
 {
 	ISourceControlProvider& SourceControlProvider = ISourceControlModule::Get().GetProvider();
 	FSourceControlChangelistStatePtr ChangelistState = SourceControlProvider.GetState(Changelist->Changelist->AsShared(), EStateCacheUsage::Use);
-		 	
+
 	// Figure out which world(s) those assets are in and split the files per world
-	TMap<FName, TArray<FAssetData>>	MapToFiles;
+	TMap<FName, TSet<FAssetData>>	MapToActorsFiles;
+	auto TryAssociateAssetDataToMap = [&MapToActorsFiles](const FAssetData& AssetData)
+	{
+		// Check that the asset is an actor
+		static FName NAME_ActorMetaDataClass(TEXT("ActorMetaDataClass"));
+		if (AssetData.TagsAndValues.Contains(NAME_ActorMetaDataClass)) // Could check AssetData.GetClass()->IsChildOf<AActor>() but this doesn't handle blueprints, all OFPA/WP Actors will have this tag
+		{
+			// WorldPartition actors are all in OFPA mode so they're external
+			// Extract the MapName from the ObjectPath (<PathToPackage>.<mapName>:<level>.<actorName>)
+			FSoftObjectPath ActorPath = FSoftObjectPath(AssetData.ObjectPath);
+			FName MapAssetName = ActorPath.GetAssetPathName();
+
+			TSet<FAssetData>* ActorFiles = MapToActorsFiles.Find(MapAssetName);
+
+			if (!ActorFiles)
+			{
+				if (ULevel::GetIsLevelPartitionedFromPackage(ActorPath.GetLongPackageFName()))
+				{
+					ActorFiles = &MapToActorsFiles.Add(MapAssetName);
+				}
+			}
 
+			if (ActorFiles)	// A null Files indicates a World not using World Partition and OFPA 
+			{
+				ActorFiles->Add(AssetData);
+				return true;
+			}
+		}
+
+		return false;
+	};
+	
+	FAssetRegistryModule& AssetRegistryModule = FModuleManager::LoadModuleChecked<FAssetRegistryModule>(FName("AssetRegistry"));
+	IAssetRegistry& AssetRegistry = AssetRegistryModule.Get();
 	for (const FSourceControlStateRef& File : ChangelistState->GetFilesStates())
 	{		
 		// Skip deleted files since we're not validating references in this validator 
@@ -85,45 +119,52 @@ EDataValidationResult UWorldPartitionChangelistValidator::ValidateActorsListFrom
 		FString PackageName;
 		if (FPackageName::TryConvertFilenameToLongPackageName(File->GetFilename(), PackageName))
 		{						
-			// Extract Actor & Map from package
 			TArray<FAssetData> PackageAssetsData;
 			USourceControlHelpers::GetAssetDataFromPackage(PackageName, PackageAssetsData);
 			
 			for (FAssetData& AssetData : PackageAssetsData)
 			{
-				// Check that the asset is an actor
-				static FName NAME_ActorMetaDataClass(TEXT("ActorMetaDataClass"));
-				if (AssetData.TagsAndValues.Contains(NAME_ActorMetaDataClass)) // Could check AssetData.GetClass()->IsChildOf<AActor>() but this doesn't handle blueprints, all OFPA/WP Actors will have this tag
+				if (AssetData.GetClass() == UDataLayerAsset::StaticClass())
 				{
-					// WorldPartition actors are all in OFPA mode so they're external
-					// Extract the MapName from the ObjectPath (<PathToPackage>.<mapName>:<level>.<actorName>)
-					FSoftObjectPath ActorPath = FSoftObjectPath(AssetData.ObjectPath);
-					FName MapAssetName = ActorPath.GetAssetPathName();
+					TArray<FName> ReferencerNames;
+					AssetRegistry.GetReferencers(AssetData.PackageName, ReferencerNames, UE::AssetRegistry::EDependencyCategory::All);
+			
+					FARFilter Filter;
+					Filter.bIncludeOnlyOnDiskAssets = true;
+					Filter.PackageNames = MoveTemp(ReferencerNames);
+					Filter.ClassNames.Add(AWorldDataLayers::StaticClass()->GetFName());
 
-					TArray<FAssetData>* Files = MapToFiles.Find(MapAssetName);
+					TArray<FAssetData> DataLayerReferencers;
+					AssetRegistry.GetAssets(Filter, DataLayerReferencers);
 
-					if (!Files)
+					for (const FAssetData& DataLayerReferencer : DataLayerReferencers)
 					{
-						if (ULevel::GetIsLevelPartitionedFromPackage(ActorPath.GetLongPackageFName()))
-						{
-							Files = &MapToFiles.Add(MapAssetName);
-						}
+						TryAssociateAssetDataToMap(DataLayerReferencer);
 					}
-
-					if (Files)	// A null Files indicates a World not using World Partition and OFPA 
+			
+					RelevantDataLayerAssets.Add(AssetData.PackageName.ToString());
+				}
+				else if (AssetData.GetClass() == AWorldDataLayers::StaticClass())
+				{
+					if (TryAssociateAssetDataToMap(AssetData))
 					{
-						Files->Add(AssetData);
+						SubmittingWorldDataLayers = true;
 					}
 				}
+				else
+				{
+
+					TryAssociateAssetDataToMap(AssetData);
+				}
 			}
 		}
 	}
 
 	// For Each world 
-	for (TTuple<FName, TArray<FAssetData>>& It : MapToFiles)
+	for (TTuple<FName, TSet<FAssetData>>& It : MapToActorsFiles)
 	{
 		const FName& MapName = It.Get<0>();
-		const TArray<FAssetData>& ActorsData = It.Get<1>();
+		const TSet<FAssetData>& ActorsData = It.Get<1>();
 
 		// Find/Load the ActorDescContainer
 		UWorld* World = FindObject<UWorld>(nullptr, *MapName.ToString(), true);
@@ -181,6 +222,12 @@ bool UWorldPartitionChangelistValidator::Filter(const FWorldPartitionActorDescVi
 	return false;
 }
 
+bool UWorldPartitionChangelistValidator::Filter(const UDataLayerInstance* InDataLayerInstance)
+{
+	const UDataLayerInstanceWithAsset* DataLayerWithAsset = Cast<UDataLayerInstanceWithAsset>(InDataLayerInstance);
+	return DataLayerWithAsset != nullptr && DataLayerWithAsset->GetAsset() != nullptr && RelevantDataLayerAssets.Contains(DataLayerWithAsset->GetAsset()->GetPathName());
+}
+
 void UWorldPartitionChangelistValidator::OnInvalidReference(const FWorldPartitionActorDescView& ActorDescView, const FGuid& ReferenceGuid)
 {
 	if (Filter(ActorDescView))
@@ -257,6 +304,47 @@ void UWorldPartitionChangelistValidator::OnInvalidReferenceLevelScriptDataLayers
 	}
 }
 
+void UWorldPartitionChangelistValidator::OnInvalidReferenceDataLayerAsset(const UDataLayerInstanceWithAsset* DataLayerInstance)
+{
+	if (SubmittingWorldDataLayers)
+	{
+		FText CurrentError = FText::Format(LOCTEXT("DataValidation.Changelist.WorldPartition.InvalidDataLayerAsset", "Data Layer {0} has no Data Layer Asset"),
+			FText::FromName(DataLayerInstance->GetDataLayerFName()));
+
+		Errors->Add(CurrentError);
+	}
+}
+
+void UWorldPartitionChangelistValidator::OnDataLayerHierarchyTypeMismatch(const UDataLayerInstance* DataLayerInstance, const UDataLayerInstance* Parent)
+{
+	if (Filter(DataLayerInstance)
+		|| Filter(Parent)
+		|| SubmittingWorldDataLayers)
+	{
+		FText CurrentError = FText::Format(LOCTEXT("DataValidation.Changelist.WorldPartition.DataLayerHierarchyTypeMismatch", "Data Layer {0} is of Type {1} and its parent {2} is of type {3}"),
+			FText::FromString(DataLayerInstance->GetDataLayerFullName()),
+			UEnum::GetDisplayValueAsText(DataLayerInstance->GetType()),
+			FText::FromString(Parent->GetDataLayerFullName()),
+			UEnum::GetDisplayValueAsText(Parent->GetType()));
+	
+		Errors->Add(CurrentError);
+	}
+}
+
+void UWorldPartitionChangelistValidator::OnDataLayerAssetConflict(const UDataLayerInstanceWithAsset* DataLayerInstance, const UDataLayerInstanceWithAsset* ConflictingDataLayerInstance)
+{
+	if (Filter(DataLayerInstance)
+		|| Filter(ConflictingDataLayerInstance)
+		|| SubmittingWorldDataLayers)
+	{
+		FText CurrentError = FText::Format(LOCTEXT("DataValidation.Changelist.WorldPartition.DataLayerAssetConflict", "Data Layer Instance {0} and Data Layer Instance {1} are both referencing Data Layer Asset {2}"),
+			FText::FromName(DataLayerInstance->GetDataLayerFName()),
+			FText::FromName(ConflictingDataLayerInstance->GetDataLayerFName()),
+			FText::FromString(DataLayerInstance->GetAsset()->GetFullName()));
+
+		Errors->Add(CurrentError);
+	}
+}
 
 
 #undef LOCTEXT_NAMESPACE
\ No newline at end of file
diff --git a/Engine/Plugins/Editor/DataValidation/Source/DataValidation/Private/WorldPartitionChangelistValidator.h b/Engine/Plugins/Editor/DataValidation/Source/DataValidation/Private/WorldPartitionChangelistValidator.h
--- a/Engine/Plugins/Editor/DataValidation/Source/DataValidation/Private/WorldPartitionChangelistValidator.h
+++ b/Engine/Plugins/Editor/DataValidation/Source/DataValidation/Private/WorldPartitionChangelistValidator.h
@@ -19,15 +19,20 @@ protected:
 	
 	TArray<FText>*							Errors = nullptr;
 	TSet<FGuid>								RelevantActorGuids;
+	TSet<FString>							RelevantDataLayerAssets;
+	bool									SubmittingWorldDataLayers = false;
 
 	virtual bool CanValidateAsset_Implementation(UObject* InAsset) const override;
 	virtual EDataValidationResult ValidateLoadedAsset_Implementation(UObject* InAsset, TArray<FText>& ValidationErrors) override;
 
-	EDataValidationResult ValidateActorsListFromChangeList(UDataValidationChangelist* Changelist);
+	EDataValidationResult ValidateActorsAndDataLayersFromChangeList(UDataValidationChangelist* Changelist);
 
 	// Return true if this ActorDescView is pertinent to the current ChangeList
 	bool Filter(const FWorldPartitionActorDescView& ActorDescView);
 
+	// Return true if this UDataLayerInstance is pertinent to the current ChangeList
+	bool Filter(const UDataLayerInstance* DataLayerInstance);
+
 	// IStreamingGenerationErrorHandler Interface methods
 	virtual void OnInvalidReference(const FWorldPartitionActorDescView& ActorDescView, const FGuid& ReferenceGuid) override;
 	virtual void OnInvalidReferenceGridPlacement(const FWorldPartitionActorDescView& ActorDescView, const FWorldPartitionActorDescView& ReferenceActorDescView) override;
@@ -35,5 +40,8 @@ protected:
 	virtual void OnInvalidReferenceLevelScriptStreamed(const FWorldPartitionActorDescView& ActorDescView) override;
 	virtual void OnInvalidReferenceLevelScriptDataLayers(const FWorldPartitionActorDescView& ActorDescView) override;
 	virtual void OnInvalidReferenceRuntimeGrid(const FWorldPartitionActorDescView& ActorDescView, const FWorldPartitionActorDescView& ReferenceActorDescView) override;
+	virtual void OnInvalidReferenceDataLayerAsset(const UDataLayerInstanceWithAsset* DataLayerInstance) override;
+	virtual void OnDataLayerHierarchyTypeMismatch(const UDataLayerInstance* DataLayerInstance, const UDataLayerInstance* Parent) override;
+	virtual void OnDataLayerAssetConflict(const UDataLayerInstanceWithAsset* DataLayerInstance, const UDataLayerInstanceWithAsset* ConflictingDataLayerInstance) override;
 };
 
diff --git a/Engine/Plugins/Editor/WorldPartitionHLODUtilities/Source/Private/WorldPartitionHLODUtilities.cpp b/Engine/Plugins/Editor/WorldPartitionHLODUtilities/Source/Private/WorldPartitionHLODUtilities.cpp
--- a/Engine/Plugins/Editor/WorldPartitionHLODUtilities/Source/Private/WorldPartitionHLODUtilities.cpp
+++ b/Engine/Plugins/Editor/WorldPartitionHLODUtilities/Source/Private/WorldPartitionHLODUtilities.cpp
@@ -6,6 +6,7 @@
 
 #include "WorldPartition/WorldPartition.h"
 #include "WorldPartition/WorldPartitionHandle.h"
+#include "WorldPartition/DataLayer/DataLayerInstance.h"
 #include "WorldPartition/HLOD/HLODActor.h"
 #include "WorldPartition/HLOD/HLODSubActor.h"
 #include "WorldPartition/HLOD/HLODActorDesc.h"
@@ -98,7 +99,7 @@ static uint32 ComputeHLODHash(AWorldPartitionHLOD* InHLODActor, const TArray<AAc
 	return Ar.GetCrc();
 }
 
-TArray<AWorldPartitionHLOD*> FWorldPartitionHLODUtilities::CreateHLODActors(FHLODCreationContext& InCreationContext, const FHLODCreationParams& InCreationParams, const TSet<FActorInstance>& InActors, const TArray<const UDataLayer*>& InDataLayers)
+TArray<AWorldPartitionHLOD*> FWorldPartitionHLODUtilities::CreateHLODActors(FHLODCreationContext& InCreationContext, const FHLODCreationParams& InCreationParams, const TSet<FActorInstance>& InActors, const TArray<const UDataLayerInstance*>& InDataLayersInstances)
 {
 	struct FSubActorsInfo
 	{
@@ -159,9 +160,9 @@ TArray<AWorldPartitionHLOD*> FWorldPartitionHLODUtilities::CreateHLODActors(FHLO
 			HLODActor->SetGridIndices(InCreationParams.GridIndexX, InCreationParams.GridIndexY, InCreationParams.GridIndexZ);
 
 			// Make sure the generated HLOD actor has the same data layers as the source actors
-			for (const UDataLayer* DataLayer : InDataLayers)
+			for (const UDataLayerInstance* DataLayerInstance : InDataLayersInstances)
 			{
-				HLODActor->AddDataLayer(DataLayer);
+				HLODActor->AddDataLayer(DataLayerInstance);
 			}
 		}
 		else
@@ -173,7 +174,7 @@ TArray<AWorldPartitionHLOD*> FWorldPartitionHLODUtilities::CreateHLODActors(FHLO
 			check(GridIndexY == InCreationParams.GridIndexY);
 			check(GridIndexZ == InCreationParams.GridIndexZ);
 			check(HLODActor->GetSubActorsHLODLayer() == HLODLayer);
-			check(FDataLayersID(HLODActor->GetDataLayerObjects()) == InCreationParams.DataLayersID);
+			check(FDataLayersID(HLODActor->GetDataLayerInstances()) == InCreationParams.DataLayersID);
 #endif
 		}
 
diff --git a/Engine/Plugins/Editor/WorldPartitionHLODUtilities/Source/Public/WorldPartitionHLODUtilities.h b/Engine/Plugins/Editor/WorldPartitionHLODUtilities/Source/Public/WorldPartitionHLODUtilities.h
--- a/Engine/Plugins/Editor/WorldPartitionHLODUtilities/Source/Public/WorldPartitionHLODUtilities.h
+++ b/Engine/Plugins/Editor/WorldPartitionHLODUtilities/Source/Public/WorldPartitionHLODUtilities.h
@@ -10,7 +10,7 @@
 class FWorldPartitionHLODUtilities : public IWorldPartitionHLODUtilities
 {
 public:
-	virtual TArray<AWorldPartitionHLOD*> CreateHLODActors(FHLODCreationContext& InCreationContext, const FHLODCreationParams& InCreationParams, const TSet<FActorInstance>& InActors, const TArray<const UDataLayer*>& InDataLayers) override;
+	virtual TArray<AWorldPartitionHLOD*> CreateHLODActors(FHLODCreationContext& InCreationContext, const FHLODCreationParams& InCreationParams, const TSet<FActorInstance>& InActors, const TArray<const UDataLayerInstance*>& InDataLayerInstances) override;
 	virtual uint32 BuildHLOD(AWorldPartitionHLOD* InHLODActor) override;
 	virtual TSubclassOf<UHLODBuilder> GetHLODBuilderClass(const UHLODLayer* InHLODLayer) override;
 	virtual UHLODBuilderSettings* CreateHLODBuilderSettings(UHLODLayer* InHLODLayer) override;
diff --git a/Engine/Plugins/LightWeightInstancesEditor/Source/LightWeightInstancesEditor/Private/LightWeightInstancesEditor.cpp b/Engine/Plugins/LightWeightInstancesEditor/Source/LightWeightInstancesEditor/Private/LightWeightInstancesEditor.cpp
--- a/Engine/Plugins/LightWeightInstancesEditor/Source/LightWeightInstancesEditor/Private/LightWeightInstancesEditor.cpp
+++ b/Engine/Plugins/LightWeightInstancesEditor/Source/LightWeightInstancesEditor/Private/LightWeightInstancesEditor.cpp
@@ -112,10 +112,10 @@ void FLightWeightInstancesEditorModule::ConvertActorsToLWIsUIAction(const TArray
 	for (AActor* Actor : InActors)
 	{
 		// use the first layer the actor is in if it's in multiple layers
-		TArray<const UDataLayer*> DataLayers = Actor->GetDataLayerObjects();
-		const UDataLayer* Layer = DataLayers.Num() > 0 ? DataLayers[0] : nullptr;
+		TArray<const UDataLayerInstance*> DataLayerInstances = Actor->GetDataLayerInstances();
+		const UDataLayerInstance* DataLayerInstance = DataLayerInstances.Num() > 0 ? DataLayerInstances[0] : nullptr;
 
-		ALightWeightInstanceManager* Manager = FLightWeightInstanceSubsystem::Get().FindOrAddLightWeightInstanceManager(Actor->GetClass(), Layer, Actor->GetWorld());
+		ALightWeightInstanceManager* Manager = FLightWeightInstanceSubsystem::Get().FindOrAddLightWeightInstanceManager(Actor->GetClass(), DataLayerInstance, Actor->GetWorld());
 		check(Manager);
 		UDataLayerEditorSubsystem::Get()->OnActorDataLayersChanged().Broadcast(Manager);
 
diff --git a/Engine/Source/Developer/AssetTools/Private/AssetTools.cpp b/Engine/Source/Developer/AssetTools/Private/AssetTools.cpp
--- a/Engine/Source/Developer/AssetTools/Private/AssetTools.cpp
+++ b/Engine/Source/Developer/AssetTools/Private/AssetTools.cpp
@@ -63,6 +63,7 @@
 #include "AssetTypeActions/AssetTypeActions_CurveLinearColor.h"
 #include "AssetTypeActions/AssetTypeActions_CurveLinearColorAtlas.h"
 #include "AssetTypeActions/AssetTypeActions_DataAsset.h"
+#include "AssetTypeActions/AssetTypeActions_DataLayer.h"
 #include "AssetTypeActions/AssetTypeActions_DataTable.h"
 #include "AssetTypeActions/AssetTypeActions_CompositeDataTable.h"
 #include "AssetTypeActions/AssetTypeActions_Enum.h"
@@ -253,6 +254,7 @@ UAssetToolsImpl::UAssetToolsImpl(const FObjectInitializer& ObjectInitializer)
 	RegisterAssetTypeActions(MakeShareable(new FAssetTypeActions_CurveLinearColor));
 	RegisterAssetTypeActions(MakeShareable(new FAssetTypeActions_CurveLinearColorAtlas));
 	RegisterAssetTypeActions(MakeShareable(new FAssetTypeActions_DataAsset));
+	RegisterAssetTypeActions(MakeShareable(new FAssetTypeActions_DataLayer));
 	RegisterAssetTypeActions(MakeShareable(new FAssetTypeActions_DataTable));
 	RegisterAssetTypeActions(MakeShareable(new FAssetTypeActions_CompositeDataTable));
 	RegisterAssetTypeActions(MakeShareable(new FAssetTypeActions_Enum));
diff --git a/Engine/Source/Developer/AssetTools/Private/AssetTypeActions/AssetTypeActions_DataLayer.h b/Engine/Source/Developer/AssetTools/Private/AssetTypeActions/AssetTypeActions_DataLayer.h
new file mode 100644
--- /dev/null
+++ b/Engine/Source/Developer/AssetTools/Private/AssetTypeActions/AssetTypeActions_DataLayer.h
@@ -0,0 +1,17 @@
+// Copyright Epic Games, Inc. All Rights Reserved.
+
+#pragma once
+
+#include "CoreMinimal.h"
+#include "AssetTypeActions_Base.h"
+#include "WorldPartition/DataLayer/DataLayerAsset.h"
+
+class FAssetTypeActions_DataLayer : public FAssetTypeActions_Base
+{
+public:
+	// IAssetTypeActions Implementation
+	virtual FText GetName() const override { return NSLOCTEXT("AssetTypeActions", "AssetTypeActions_DataLayer", "Data Layer"); }
+	virtual FColor GetTypeColor() const override { return FColor(12, 65, 12); }
+	virtual UClass* GetSupportedClass() const override { return UDataLayerAsset::StaticClass(); }
+	virtual uint32 GetCategories() override { return EAssetTypeCategories::Misc; }
+};
\ No newline at end of file
diff --git a/Engine/Source/Editor/DataLayerEditor/Private/DataLayer/DataLayerActorTreeItem.h b/Engine/Source/Editor/DataLayerEditor/Private/DataLayer/DataLayerActorTreeItem.h
--- a/Engine/Source/Editor/DataLayerEditor/Private/DataLayer/DataLayerActorTreeItem.h
+++ b/Engine/Source/Editor/DataLayerEditor/Private/DataLayer/DataLayerActorTreeItem.h
@@ -6,19 +6,19 @@
 #include "GameFramework/Actor.h"
 #include "SceneOutlinerFwd.h"
 #include "ActorTreeItem.h"
-#include "WorldPartition/DataLayer/DataLayer.h"
+#include "WorldPartition/DataLayer/DataLayerInstance.h"
 
 //////////////////////////////////////////////////////////////////////////
 // FDataLayerActorTreeItemData
 
 struct FDataLayerActorTreeItemData
 {
-	FDataLayerActorTreeItemData(AActor* InActor, UDataLayer* InDataLayer)
+	FDataLayerActorTreeItemData(AActor* InActor, UDataLayerInstance* InDataLayerInstance)
 		: Actor(InActor)
-		, DataLayer(InDataLayer)
+		, DataLayerInstance(InDataLayerInstance)
 	{}
 	TWeakObjectPtr<AActor> Actor;
-	TWeakObjectPtr<UDataLayer> DataLayer;
+	TWeakObjectPtr<UDataLayerInstance> DataLayerInstance;
 };
 
 //////////////////////////////////////////////////////////////////////////
@@ -27,38 +27,38 @@ struct FDataLayerActorTreeItemData
 struct FDataLayerActorTreeItem : public FActorTreeItem
 {
 public:
-	DECLARE_DELEGATE_RetVal_TwoParams(bool, FFilterPredicate, const AActor*, const UDataLayer* DataLayer);
-	DECLARE_DELEGATE_RetVal_TwoParams(bool, FInteractivePredicate, const AActor*, const UDataLayer* DataLayer);
+	DECLARE_DELEGATE_RetVal_TwoParams(bool, FFilterPredicate, const AActor*, const UDataLayerInstance* InDataLayerInstance);
+	DECLARE_DELEGATE_RetVal_TwoParams(bool, FInteractivePredicate, const AActor*, const UDataLayerInstance* InDataLayerInstance);
 
 	FDataLayerActorTreeItem(const FDataLayerActorTreeItemData& InData)
 		: FActorTreeItem(InData.Actor.Get())
-		, DataLayer(InData.DataLayer)
-		, IDDataLayerActor(FDataLayerActorTreeItem::ComputeTreeItemID(Actor.Get(), DataLayer.Get()))
+		, DataLayerInstance(InData.DataLayerInstance)
+		, IDDataLayerActor(FDataLayerActorTreeItem::ComputeTreeItemID(Actor.Get(), DataLayerInstance.Get()))
 	{
 	}
 
-	UDataLayer* GetDataLayer() const { return DataLayer.Get(); }
+	UDataLayerInstance* GetDataLayer() const { return DataLayerInstance.Get(); }
 	
 	const AActor* GetActor() const { return Actor.Get(); }
 	AActor* GetActor() { return Actor.Get(); }
 
-	static uint32 ComputeTreeItemID(const AActor* InActor, const UDataLayer* InDataLayer)
+	static uint32 ComputeTreeItemID(const AActor* InActor, const UDataLayerInstance* InDataLayerInstance)
 	{
-		return HashCombine(GetTypeHash(FObjectKey(InActor)), GetTypeHash(FObjectKey(InDataLayer)));
+		return HashCombine(GetTypeHash(FObjectKey(InActor)), GetTypeHash(FObjectKey(InDataLayerInstance)));
 	}
 
 	bool Filter(FFilterPredicate Pred) const
 	{
-		return Pred.Execute(Actor.Get(), DataLayer.Get());
+		return Pred.Execute(Actor.Get(), DataLayerInstance.Get());
 	}
 
 	bool GetInteractiveState(FInteractivePredicate Pred) const
 	{
-		return Pred.Execute(Actor.Get(), DataLayer.Get());
+		return Pred.Execute(Actor.Get(), DataLayerInstance.Get());
 	}
 
 	/* Begin ISceneOutlinerTreeItem Implementation */
-	virtual bool IsValid() const override { return Actor.IsValid() && DataLayer.IsValid(); }
+	virtual bool IsValid() const override { return Actor.IsValid() && DataLayerInstance.IsValid(); }
 	virtual FSceneOutlinerTreeItemID GetID() const override { return IDDataLayerActor; }
 	virtual bool ShouldShowVisibilityState() const { return false; }
 	virtual bool HasVisibilityInfo() const override { return false; }
@@ -67,6 +67,6 @@ public:
 	/* End ISceneOutlinerTreeItem Implementation */
 
 private:
-	TWeakObjectPtr<UDataLayer> DataLayer;
+	TWeakObjectPtr<UDataLayerInstance> DataLayerInstance;
 	const uint32 IDDataLayerActor;
 };
\ No newline at end of file
diff --git a/Engine/Source/Editor/DataLayerEditor/Private/DataLayer/DataLayerDragDropOp.cpp b/Engine/Source/Editor/DataLayerEditor/Private/DataLayer/DataLayerDragDropOp.cpp
--- a/Engine/Source/Editor/DataLayerEditor/Private/DataLayer/DataLayerDragDropOp.cpp
+++ b/Engine/Source/Editor/DataLayerEditor/Private/DataLayer/DataLayerDragDropOp.cpp
@@ -9,13 +9,13 @@
 void FDataLayerDragDropOp::Construct()
 {
 	const FSlateBrush* Icon = FEditorStyle::GetBrush(TEXT("DataLayer.Editor"));
-	if (DataLayerLabels.Num() == 1)
+	if (DataLayerDragDropInfos.Num() == 1)
 	{
-		SetToolTip(FText::FromName(DataLayerLabels[0]), Icon);
+		SetToolTip(FText::FromString(DataLayerDragDropInfos[0].DataLayerShortName), Icon);
 	}
 	else
 	{
-		FText Text = FText::Format(NSLOCTEXT("FDataLayerDragDropOp", "MultipleFormat", "{0} DataLayerLabels"), DataLayerLabels.Num());
+		FText Text = FText::Format(NSLOCTEXT("FDataLayerDragDropOp", "MultipleFormat", "{0} DataLayers"), DataLayerDragDropInfos.Num());
 		SetToolTip(Text, Icon);
 	}
 
diff --git a/Engine/Source/Editor/DataLayerEditor/Private/DataLayer/DataLayerDragDropOp.h b/Engine/Source/Editor/DataLayerEditor/Private/DataLayer/DataLayerDragDropOp.h
--- a/Engine/Source/Editor/DataLayerEditor/Private/DataLayer/DataLayerDragDropOp.h
+++ b/Engine/Source/Editor/DataLayerEditor/Private/DataLayer/DataLayerDragDropOp.h
@@ -4,11 +4,13 @@
 
 #include "DragAndDrop/DecoratedDragDropOp.h"
 
+#include "WorldPartition/DataLayer/DataLayerInstance.h"
+
 class FDragDropEvent;
 class AActor;
-class UDataLayer;
+class UDataLayerInstance;
 
-typedef TPair< TWeakObjectPtr<AActor>, TWeakObjectPtr<UDataLayer>> FDataLayerActorMoveElement;
+typedef TPair< TWeakObjectPtr<AActor>, TWeakObjectPtr<UDataLayerInstance>> FDataLayerActorMoveElement;
 
 /** Drag/drop operation for dragging layers in the editor */
 class FDataLayerDragDropOp : public FDecoratedDragDropOp
@@ -16,8 +18,19 @@ class FDataLayerDragDropOp : public FDecoratedDragDropOp
 public:
 	DRAG_DROP_OPERATOR_TYPE(FDataLayerDragDropOp, FDecoratedDragDropOp)
 
+	struct FDragDropInfo
+	{
+		FDragDropInfo(const UDataLayerInstance* DataLayerInstance)
+			:DataLayerInstanceName(DataLayerInstance->GetDataLayerFName()),
+			DataLayerShortName(DataLayerInstance->GetDataLayerShortName())
+		{}
+
+		FName DataLayerInstanceName;
+		FString DataLayerShortName;
+	};
+
 	/** The labels of the layers being dragged */
-	TArray<FName> DataLayerLabels;
+	TArray<FDragDropInfo> DataLayerDragDropInfos;
 
 	virtual void Construct() override;
 };
diff --git a/Engine/Source/Editor/DataLayerEditor/Private/DataLayer/DataLayerEditorSubsystem.cpp b/Engine/Source/Editor/DataLayerEditor/Private/DataLayer/DataLayerEditorSubsystem.cpp
--- a/Engine/Source/Editor/DataLayerEditor/Private/DataLayer/DataLayerEditorSubsystem.cpp
+++ b/Engine/Source/Editor/DataLayerEditor/Private/DataLayer/DataLayerEditorSubsystem.cpp
@@ -2,12 +2,15 @@
 
 #include "DataLayer/DataLayerEditorSubsystem.h"
 #include "Components/PrimitiveComponent.h"
-#include "WorldPartition/DataLayer/DataLayer.h"
+#include "WorldPartition/DataLayer/DataLayerInstanceWithAsset.h"
+#include "WorldPartition/DataLayer/DeprecatedDataLayerInstance.h"
 #include "WorldPartition/DataLayer/WorldDataLayers.h"
+#include "WorldPartition/DataLayer/DataLayerSubsystem.h"
 #include "WorldPartition/WorldPartitionSubsystem.h"
 #include "WorldPartition/WorldPartition.h"
 #include "WorldPartition/WorldPartitionEditorPerProjectUserSettings.h"
 #include "Modules/ModuleManager.h"
+#include "Misc/MessageDialog.h"
 #include "DataLayerEditorModule.h"
 #include "ActorEditorUtils.h"
 #include "LevelEditorViewport.h"
@@ -19,6 +22,10 @@
 #include "EngineUtils.h"
 #include "Editor.h"
 #include "Subsystems/ActorEditorContextSubsystem.h"
+#include "IContentBrowserSingleton.h"
+#include "ContentBrowserModule.h"
+
+#define LOCTEXT_NAMESPACE "DataLayer"
 
 //////////////////////////////////////////////////////////////////////////
 // FDataLayersBroadcast
@@ -87,7 +94,7 @@ void FDataLayersBroadcast::OnObjectPostEditChange(UObject* Object, FPropertyChan
 	if (Object && (PropertyChangedEvent.ChangeType != EPropertyChangeType::Interactive))
 	{
 		bool bRefresh = false;
-		if (UDataLayer* DataLayer = Cast<UDataLayer>(Object))
+		if (Object->IsA<UDataLayerInstance>() || Object->IsA<UDataLayerAsset>())
 		{
 			bRefresh = true;
 		}
@@ -161,7 +168,7 @@ void UDataLayerEditorSubsystem::OnExecuteActorEditorContextAction(UWorld* InWorl
 		}
 		break;
 	case EActorEditorContextAction::ResetContext:
-		for (UDataLayer* DataLayer : WorldDataLayer->GetActorEditorContextDataLayers())
+		for (UDataLayerInstance* DataLayer : WorldDataLayer->GetActorEditorContextDataLayers())
 		{
 			RemoveFromActorEditorContext(DataLayer);
 		}
@@ -194,8 +201,8 @@ TSharedRef<SWidget> UDataLayerEditorSubsystem::GetActorEditorContextWidget(UWorl
 	TSharedRef<SVerticalBox> OutWidget = SNew(SVerticalBox);
 	if (AWorldDataLayers* WorldDataLayer = InWorld->GetWorldDataLayers())
 	{
-		TArray<UDataLayer*> DataLayers = WorldDataLayer->GetActorEditorContextDataLayers();
-		for (UDataLayer* DataLayer : DataLayers)
+		TArray<UDataLayerInstance*> DataLayers = WorldDataLayer->GetActorEditorContextDataLayers();
+		for (UDataLayerInstance* DataLayer : DataLayers)
 		{
 			check(IsValid(DataLayer));
 			OutWidget->AddSlot().AutoHeight()
@@ -214,7 +221,7 @@ TSharedRef<SWidget> UDataLayerEditorSubsystem::GetActorEditorContextWidget(UWorl
 				.Padding(4.0f, 1.0f, 1.0f, 1.0f)
 				[
 					SNew(STextBlock)
-					.Text(FText::FromString(DataLayer->GetDataLayerLabel().ToString()))
+					.Text(FText::FromString(DataLayer->GetDataLayerShortName()))
 				]
 			];
 		}
@@ -222,19 +229,19 @@ TSharedRef<SWidget> UDataLayerEditorSubsystem::GetActorEditorContextWidget(UWorl
 	return OutWidget;
 }
 
-void UDataLayerEditorSubsystem::AddToActorEditorContext(UDataLayer* InDataLayer)
+void UDataLayerEditorSubsystem::AddToActorEditorContext(UDataLayerInstance* InDataLayerInstance)
 {
-	if (InDataLayer->AddToActorEditorContext())
+	if (InDataLayerInstance->AddToActorEditorContext())
 	{
-		BroadcastDataLayerChanged(EDataLayerAction::Modify, InDataLayer, NAME_None);
+		BroadcastDataLayerChanged(EDataLayerAction::Modify, InDataLayerInstance, NAME_None);
 	}
 }
 
-void UDataLayerEditorSubsystem::RemoveFromActorEditorContext(UDataLayer* InDataLayer)
+void UDataLayerEditorSubsystem::RemoveFromActorEditorContext(UDataLayerInstance* InDataLayerInstance)
 {
-	if (InDataLayer->RemoveFromActorEditorContext())
+	if (InDataLayerInstance->RemoveFromActorEditorContext())
 	{
-		BroadcastDataLayerChanged(EDataLayerAction::Modify, InDataLayer, NAME_None);
+		BroadcastDataLayerChanged(EDataLayerAction::Modify, InDataLayerInstance, NAME_None);
 	}
 }
 
@@ -315,7 +322,7 @@ UWorld* UDataLayerEditorSubsystem::GetWorld() const
 	return GWorld;
 }
 
-bool UDataLayerEditorSubsystem::SetParentDataLayer(UDataLayer* DataLayer, UDataLayer* ParentDataLayer)
+bool UDataLayerEditorSubsystem::SetParentDataLayer(UDataLayerInstance* DataLayer, UDataLayerInstance* ParentDataLayer)
 {
 	if (DataLayer->CanParent(ParentDataLayer))
 	{
@@ -332,7 +339,7 @@ bool UDataLayerEditorSubsystem::SetParentDataLayer(UDataLayer* DataLayer, UDataL
 	return false;
 }
 
-bool UDataLayerEditorSubsystem::AddActorToDataLayer(AActor* Actor, UDataLayer* DataLayer)
+bool UDataLayerEditorSubsystem::AddActorToDataLayer(AActor* Actor, UDataLayerInstance* DataLayer)
 {
 	TArray<AActor*> Actors;
 	Actors.Add(Actor);
@@ -340,7 +347,7 @@ bool UDataLayerEditorSubsystem::AddActorToDataLayer(AActor* Actor, UDataLayer* D
 	return AddActorsToDataLayers(Actors, {DataLayer});
 }
 
-bool UDataLayerEditorSubsystem::AddActorToDataLayers(AActor* Actor, const TArray<UDataLayer*>& DataLayers)
+bool UDataLayerEditorSubsystem::AddActorToDataLayers(AActor* Actor, const TArray<UDataLayerInstance*>& DataLayers)
 {
 	TArray<AActor*> Actors;
 	Actors.Add(Actor);
@@ -348,12 +355,12 @@ bool UDataLayerEditorSubsystem::AddActorToDataLayers(AActor* Actor, const TArray
 	return AddActorsToDataLayers(Actors, DataLayers);
 }
 
-bool UDataLayerEditorSubsystem::AddActorsToDataLayer(const TArray<AActor*>& Actors, UDataLayer* DataLayer)
+bool UDataLayerEditorSubsystem::AddActorsToDataLayer(const TArray<AActor*>& Actors, UDataLayerInstance* DataLayer)
 {
 	return AddActorsToDataLayers(Actors, {DataLayer});
 }
 
-bool UDataLayerEditorSubsystem::AddActorsToDataLayers(const TArray<AActor*>& Actors, const TArray<UDataLayer*>& DataLayers)
+bool UDataLayerEditorSubsystem::AddActorsToDataLayers(const TArray<AActor*>& Actors, const TArray<UDataLayerInstance*>& DataLayers)
 {
 	bool bChangesOccurred = false;
 
@@ -369,7 +376,7 @@ bool UDataLayerEditorSubsystem::AddActorsToDataLayers(const TArray<AActor*>& Act
 			}
 
 			bool bActorWasModified = false;
-			for (const UDataLayer* DataLayer : DataLayers)
+			for (const UDataLayerInstance* DataLayer : DataLayers)
 			{
 				if (Actor->AddDataLayer(DataLayer))
 				{
@@ -407,12 +414,12 @@ bool UDataLayerEditorSubsystem::RemoveActorsFromAllDataLayers(const TArray<AActo
 	bool bChangesOccurred = false;
 	for (AActor* Actor : Actors)
 	{
-		TArray<const UDataLayer*> ModifiedDataLayers = Actor->GetDataLayerObjects();
+		TArray<const UDataLayerInstance*> ModifiedDataLayerInstances = Actor->GetDataLayerInstances();
 		if (Actor->RemoveAllDataLayers())
 		{
-			for (const UDataLayer* DataLayer : ModifiedDataLayers)
+			for (const UDataLayerInstance* DataLayerInstance : ModifiedDataLayerInstances)
 			{
-				BroadcastDataLayerChanged(EDataLayerAction::Modify, DataLayer, NAME_None);
+				BroadcastDataLayerChanged(EDataLayerAction::Modify, DataLayerInstance, NAME_None);
 			}
 			BroadcastActorDataLayersChanged(Actor);
 
@@ -430,7 +437,7 @@ bool UDataLayerEditorSubsystem::RemoveActorsFromAllDataLayers(const TArray<AActo
 	return bChangesOccurred;
 }
 
-bool UDataLayerEditorSubsystem::RemoveActorFromDataLayer(AActor* Actor, UDataLayer* DataLayer)
+bool UDataLayerEditorSubsystem::RemoveActorFromDataLayer(AActor* Actor, UDataLayerInstance* DataLayer)
 {
 	TArray<AActor*> Actors;
 	Actors.Add(Actor);
@@ -438,7 +445,7 @@ bool UDataLayerEditorSubsystem::RemoveActorFromDataLayer(AActor* Actor, UDataLay
 	return RemoveActorsFromDataLayers(Actors, {DataLayer});
 }
 
-bool UDataLayerEditorSubsystem::RemoveActorFromDataLayers(AActor* Actor, const TArray<UDataLayer*>& DataLayers)
+bool UDataLayerEditorSubsystem::RemoveActorFromDataLayers(AActor* Actor, const TArray<UDataLayerInstance*>& DataLayers)
 {
 	TArray<AActor*> Actors;
 	Actors.Add(Actor);
@@ -446,12 +453,12 @@ bool UDataLayerEditorSubsystem::RemoveActorFromDataLayers(AActor* Actor, const T
 	return RemoveActorsFromDataLayers(Actors, DataLayers);
 }
 
-bool UDataLayerEditorSubsystem::RemoveActorsFromDataLayer(const TArray<AActor*>& Actors, UDataLayer* DataLayer)
+bool UDataLayerEditorSubsystem::RemoveActorsFromDataLayer(const TArray<AActor*>& Actors, UDataLayerInstance* DataLayer)
 {
 	return RemoveActorsFromDataLayers(Actors, {DataLayer});
 }
 
-bool UDataLayerEditorSubsystem::RemoveActorsFromDataLayers(const TArray<AActor*>& Actors, const TArray<UDataLayer*>& DataLayers)
+bool UDataLayerEditorSubsystem::RemoveActorsFromDataLayers(const TArray<AActor*>& Actors, const TArray<UDataLayerInstance*>& DataLayers)
 {
 	GEditor->GetSelectedActors()->BeginBatchSelectOperation();
 
@@ -464,7 +471,7 @@ bool UDataLayerEditorSubsystem::RemoveActorsFromDataLayers(const TArray<AActor*>
 		}
 
 		bool bActorWasModified = false;
-		for (const UDataLayer* DataLayer : DataLayers)
+		for (const UDataLayerInstance* DataLayer : DataLayers)
 		{
 			if (Actor->RemoveDataLayer(DataLayer))
 			{
@@ -503,22 +510,22 @@ TArray<AActor*> UDataLayerEditorSubsystem::GetSelectedActors() const
 	return CurrentlySelectedActors;
 }
 
-bool UDataLayerEditorSubsystem::AddSelectedActorsToDataLayer(UDataLayer* DataLayer)
+bool UDataLayerEditorSubsystem::AddSelectedActorsToDataLayer(UDataLayerInstance* DataLayer)
 {
 	return AddActorsToDataLayer(GetSelectedActors(), DataLayer);
 }
 
-bool UDataLayerEditorSubsystem::RemoveSelectedActorsFromDataLayer(UDataLayer* DataLayer)
+bool UDataLayerEditorSubsystem::RemoveSelectedActorsFromDataLayer(UDataLayerInstance* DataLayer)
 {
 	return RemoveActorsFromDataLayer(GetSelectedActors(), DataLayer);
 }
 
-bool UDataLayerEditorSubsystem::AddSelectedActorsToDataLayers(const TArray<UDataLayer*>& DataLayers)
+bool UDataLayerEditorSubsystem::AddSelectedActorsToDataLayers(const TArray<UDataLayerInstance*>& DataLayers)
 {
 	return AddActorsToDataLayers(GetSelectedActors(), DataLayers);
 }
 
-bool UDataLayerEditorSubsystem::RemoveSelectedActorsFromDataLayers(const TArray<UDataLayer*>& DataLayers)
+bool UDataLayerEditorSubsystem::RemoveSelectedActorsFromDataLayers(const TArray<UDataLayerInstance*>& DataLayers)
 {
 	return RemoveActorsFromDataLayers(GetSelectedActors(), DataLayers);
 }
@@ -529,12 +536,12 @@ bool UDataLayerEditorSubsystem::RemoveSelectedActorsFromDataLayers(const TArray<
 // Operations on actors in DataLayers
 //
 ///////////////////////////////////////////////////////////////////////////////////////////////////////////////////////
-bool UDataLayerEditorSubsystem::SelectActorsInDataLayer(UDataLayer* DataLayer, const bool bSelect, const bool bNotify, const bool bSelectEvenIfHidden)
+bool UDataLayerEditorSubsystem::SelectActorsInDataLayer(UDataLayerInstance* DataLayer, const bool bSelect, const bool bNotify, const bool bSelectEvenIfHidden)
 {
 	return SelectActorsInDataLayer(DataLayer, bSelect, bNotify, bSelectEvenIfHidden, nullptr);
 }
 
-bool UDataLayerEditorSubsystem::SelectActorsInDataLayer(UDataLayer* DataLayer, const bool bSelect, const bool bNotify, const bool bSelectEvenIfHidden, const TSharedPtr<FActorFilter>& Filter)
+bool UDataLayerEditorSubsystem::SelectActorsInDataLayer(UDataLayerInstance* DataLayer, const bool bSelect, const bool bNotify, const bool bSelectEvenIfHidden, const TSharedPtr<FActorFilter>& Filter)
 {
 	bool bChangesOccurred = false;
 
@@ -573,12 +580,12 @@ bool UDataLayerEditorSubsystem::SelectActorsInDataLayer(UDataLayer* DataLayer, c
 	return bChangesOccurred;
 }
 
-bool UDataLayerEditorSubsystem::SelectActorsInDataLayers(const TArray<UDataLayer*>& DataLayers, const bool bSelect, const bool bNotify, const bool bSelectEvenIfHidden)
+bool UDataLayerEditorSubsystem::SelectActorsInDataLayers(const TArray<UDataLayerInstance*>& DataLayers, const bool bSelect, const bool bNotify, const bool bSelectEvenIfHidden)
 {
 	return SelectActorsInDataLayers(DataLayers, bSelect, bNotify, bSelectEvenIfHidden, nullptr);
 }
 
-bool UDataLayerEditorSubsystem::SelectActorsInDataLayers(const TArray<UDataLayer*>& DataLayers, const bool bSelect, const bool bNotify, const bool bSelectEvenIfHidden, const TSharedPtr<FActorFilter>& Filter)
+bool UDataLayerEditorSubsystem::SelectActorsInDataLayers(const TArray<UDataLayerInstance*>& DataLayers, const bool bSelect, const bool bNotify, const bool bSelectEvenIfHidden, const TSharedPtr<FActorFilter>& Filter)
 {
 	if (DataLayers.Num() == 0)
 	{
@@ -601,7 +608,7 @@ bool UDataLayerEditorSubsystem::SelectActorsInDataLayers(const TArray<UDataLayer
 			continue;
 		}
 
-		for (const UDataLayer* DataLayer : DataLayers)
+		for (const UDataLayerInstance* DataLayer : DataLayers)
 		{
 			if (Actor->ContainsDataLayer(DataLayer))
 			{
@@ -652,7 +659,7 @@ bool UDataLayerEditorSubsystem::UpdateActorVisibility(AActor* Actor, bool& bOutS
 	bool bActorBelongsToVisibleDataLayer = false;
 	if (AWorldDataLayers* WorldDataLayers = GetWorldDataLayers())
 	{
-		WorldDataLayers->ForEachDataLayer([&Actor, &bOutActorModified, &bActorBelongsToVisibleDataLayer](UDataLayer* DataLayer)
+		WorldDataLayers->ForEachDataLayer([&Actor, &bOutActorModified, &bActorBelongsToVisibleDataLayer](UDataLayerInstance* DataLayer)
 		{
 			if (DataLayer->IsEffectiveVisible() && Actor->ContainsDataLayer(DataLayer))
 			{
@@ -736,12 +743,12 @@ bool UDataLayerEditorSubsystem::UpdateAllActorsVisibility(const bool bNotifySele
 //
 ///////////////////////////////////////////////////////////////////////////////////////////////////////////////////////
 
-void UDataLayerEditorSubsystem::AppendActorsFromDataLayer(UDataLayer* DataLayer, TArray<AActor*>& InOutActors) const
+void UDataLayerEditorSubsystem::AppendActorsFromDataLayer(UDataLayerInstance* DataLayer, TArray<AActor*>& InOutActors) const
 {
 	AppendActorsFromDataLayer(DataLayer, InOutActors, nullptr);
 }
 
-void UDataLayerEditorSubsystem::AppendActorsFromDataLayer(UDataLayer* DataLayer, TArray<AActor*>& InOutActors, const TSharedPtr<FActorFilter>& Filter) const
+void UDataLayerEditorSubsystem::AppendActorsFromDataLayer(UDataLayerInstance* DataLayer, TArray<AActor*>& InOutActors, const TSharedPtr<FActorFilter>& Filter) const
 {
 	for (AActor* Actor : FActorRange(GetWorld()))
 	{
@@ -756,12 +763,12 @@ void UDataLayerEditorSubsystem::AppendActorsFromDataLayer(UDataLayer* DataLayer,
 	}
 }
 
-void UDataLayerEditorSubsystem::AppendActorsFromDataLayers(const TArray<UDataLayer*>& DataLayers, TArray<AActor*>& InOutActors) const
+void UDataLayerEditorSubsystem::AppendActorsFromDataLayers(const TArray<UDataLayerInstance*>& DataLayers, TArray<AActor*>& InOutActors) const
 {
 	AppendActorsFromDataLayers(DataLayers, InOutActors, nullptr);
 }
 
-void UDataLayerEditorSubsystem::AppendActorsFromDataLayers(const TArray<UDataLayer*>& DataLayers, TArray<AActor*>& InOutActors, const TSharedPtr<FActorFilter>& Filter) const
+void UDataLayerEditorSubsystem::AppendActorsFromDataLayers(const TArray<UDataLayerInstance*>& DataLayers, TArray<AActor*>& InOutActors, const TSharedPtr<FActorFilter>& Filter) const
 {
 	for (AActor* Actor : FActorRange(GetWorld()))
 	{
@@ -769,7 +776,7 @@ void UDataLayerEditorSubsystem::AppendActorsFromDataLayers(const TArray<UDataLay
 		{
 			continue;
 		}
-		for (const UDataLayer* DataLayer : DataLayers)
+		for (const UDataLayerInstance* DataLayer : DataLayers)
 		{
 			if (Actor->ContainsDataLayer(DataLayer))
 			{
@@ -780,43 +787,43 @@ void UDataLayerEditorSubsystem::AppendActorsFromDataLayers(const TArray<UDataLay
 	}
 }
 
-TArray<AActor*> UDataLayerEditorSubsystem::GetActorsFromDataLayer(UDataLayer* DataLayer) const
+TArray<AActor*> UDataLayerEditorSubsystem::GetActorsFromDataLayer(UDataLayerInstance* DataLayer) const
 {
 	TArray<AActor*> OutActors;
 	AppendActorsFromDataLayer(DataLayer, OutActors);
 	return OutActors;
 }
 
-TArray<AActor*> UDataLayerEditorSubsystem::GetActorsFromDataLayer(UDataLayer* DataLayer, const TSharedPtr<FActorFilter>& Filter) const
+TArray<AActor*> UDataLayerEditorSubsystem::GetActorsFromDataLayer(UDataLayerInstance* DataLayer, const TSharedPtr<FActorFilter>& Filter) const
 {
 	TArray<AActor*> OutActors;
 	AppendActorsFromDataLayer(DataLayer, OutActors, Filter);
 	return OutActors;
 }
 
-TArray<AActor*> UDataLayerEditorSubsystem::GetActorsFromDataLayers(const TArray<UDataLayer*>& DataLayers) const
+TArray<AActor*> UDataLayerEditorSubsystem::GetActorsFromDataLayers(const TArray<UDataLayerInstance*>& DataLayers) const
 {
 	TArray<AActor*> OutActors;
 	AppendActorsFromDataLayers(DataLayers, OutActors);
 	return OutActors;
 }
 
-TArray<AActor*> UDataLayerEditorSubsystem::GetActorsFromDataLayers(const TArray<UDataLayer*>& DataLayers, const TSharedPtr<FActorFilter>& Filter) const
+TArray<AActor*> UDataLayerEditorSubsystem::GetActorsFromDataLayers(const TArray<UDataLayerInstance*>& DataLayers, const TSharedPtr<FActorFilter>& Filter) const
 {
 	TArray<AActor*> OutActors;
 	AppendActorsFromDataLayers(DataLayers, OutActors, Filter);
 	return OutActors;
 }
 
-void UDataLayerEditorSubsystem::SetDataLayerVisibility(UDataLayer* DataLayer, const bool bIsVisible)
+void UDataLayerEditorSubsystem::SetDataLayerVisibility(UDataLayerInstance* DataLayer, const bool bIsVisible)
 {
 	SetDataLayersVisibility({ DataLayer }, bIsVisible);
 }
 
-void UDataLayerEditorSubsystem::SetDataLayersVisibility(const TArray<UDataLayer*>& DataLayers, const bool bIsVisible)
+void UDataLayerEditorSubsystem::SetDataLayersVisibility(const TArray<UDataLayerInstance*>& DataLayers, const bool bIsVisible)
 {
 	bool bChangeOccurred = false;
-	for (UDataLayer* DataLayer : DataLayers)
+	for (UDataLayerInstance* DataLayer : DataLayers)
 	{
 		check(DataLayer);
 
@@ -835,20 +842,20 @@ void UDataLayerEditorSubsystem::SetDataLayersVisibility(const TArray<UDataLayer*
 	}
 }
 
-void UDataLayerEditorSubsystem::ToggleDataLayerVisibility(UDataLayer* DataLayer)
+void UDataLayerEditorSubsystem::ToggleDataLayerVisibility(UDataLayerInstance* DataLayer)
 {
 	check(DataLayer);
 	SetDataLayerVisibility(DataLayer, !DataLayer->IsVisible());
 }
 
-void UDataLayerEditorSubsystem::ToggleDataLayersVisibility(const TArray<UDataLayer*>& DataLayers)
+void UDataLayerEditorSubsystem::ToggleDataLayersVisibility(const TArray<UDataLayerInstance*>& DataLayers)
 {
 	if (DataLayers.Num() == 0)
 	{
 		return;
 	}
 
-	for (UDataLayer* DataLayer : DataLayers)
+	for (UDataLayerInstance* DataLayer : DataLayers)
 	{
 		DataLayer->Modify();
 		DataLayer->SetVisible(!DataLayer->IsVisible());
@@ -862,7 +869,7 @@ void UDataLayerEditorSubsystem::MakeAllDataLayersVisible()
 {
 	if (AWorldDataLayers* WorldDataLayers = GetWorldDataLayers())
 	{
-		WorldDataLayers->ForEachDataLayer([this](UDataLayer* DataLayer)
+		WorldDataLayers->ForEachDataLayer([this](UDataLayerInstance* DataLayer)
 		{
 			if (!DataLayer->IsVisible())
 			{
@@ -877,7 +884,7 @@ void UDataLayerEditorSubsystem::MakeAllDataLayersVisible()
 	UpdateAllActorsVisibility(true, true);
 }
 
-bool UDataLayerEditorSubsystem::SetDataLayerIsLoadedInEditorInternal(UDataLayer* DataLayer, const bool bIsLoadedInEditor, const bool bIsFromUserChange)
+bool UDataLayerEditorSubsystem::SetDataLayerIsLoadedInEditorInternal(UDataLayerInstance* DataLayer, const bool bIsLoadedInEditor, const bool bIsFromUserChange)
 {
 	check(DataLayer);
 	if (DataLayer->IsLoadedInEditor() != bIsLoadedInEditor)
@@ -897,32 +904,32 @@ bool UDataLayerEditorSubsystem::SetDataLayerIsLoadedInEditorInternal(UDataLayer*
 	return false;
 }
 
-bool UDataLayerEditorSubsystem::SetDataLayerIsLoadedInEditor(UDataLayer* DataLayer, const bool bIsLoadedInEditor, const bool bIsFromUserChange)
+bool UDataLayerEditorSubsystem::SetDataLayerIsLoadedInEditor(UDataLayerInstance* DataLayer, const bool bIsLoadedInEditor, const bool bIsFromUserChange)
 {
 	bool bRefreshNeeded = SetDataLayerIsLoadedInEditorInternal(DataLayer, bIsLoadedInEditor, bIsFromUserChange);
 	return bRefreshNeeded ? RefreshWorldPartitionEditorCells(bIsFromUserChange) : true;
 }
 
-bool UDataLayerEditorSubsystem::SetDataLayersIsLoadedInEditor(const TArray<UDataLayer*>& DataLayers, const bool bIsLoadedInEditor, const bool bIsFromUserChange)
+bool UDataLayerEditorSubsystem::SetDataLayersIsLoadedInEditor(const TArray<UDataLayerInstance*>& DataLayers, const bool bIsLoadedInEditor, const bool bIsFromUserChange)
 {
 	bool bRefreshNeeded = false;
-	for (UDataLayer* DataLayer : DataLayers)
+	for (UDataLayerInstance* DataLayer : DataLayers)
 	{
 		bRefreshNeeded |= SetDataLayerIsLoadedInEditorInternal(DataLayer, bIsLoadedInEditor, bIsFromUserChange);
 	}
 	return bRefreshNeeded ? RefreshWorldPartitionEditorCells(bIsFromUserChange) : true;
 }
 
-bool UDataLayerEditorSubsystem::ToggleDataLayerIsLoadedInEditor(UDataLayer* DataLayer, const bool bIsFromUserChange)
+bool UDataLayerEditorSubsystem::ToggleDataLayerIsLoadedInEditor(UDataLayerInstance* DataLayer, const bool bIsFromUserChange)
 {
 	check(DataLayer);
 	return SetDataLayerIsLoadedInEditor(DataLayer, !DataLayer->IsLoadedInEditor(), bIsFromUserChange);
 }
 
-bool UDataLayerEditorSubsystem::ToggleDataLayersIsLoadedInEditor(const TArray<UDataLayer*>& DataLayers, const bool bIsFromUserChange)
+bool UDataLayerEditorSubsystem::ToggleDataLayersIsLoadedInEditor(const TArray<UDataLayerInstance*>& DataLayers, const bool bIsFromUserChange)
 {
 	bool bRefreshNeeded = false;
-	for (UDataLayer* DataLayer : DataLayers)
+	for (UDataLayerInstance* DataLayer : DataLayers)
 	{
 		bRefreshNeeded |= SetDataLayerIsLoadedInEditorInternal(DataLayer, !DataLayer->IsLoadedInEditor(), bIsFromUserChange);
 	}
@@ -934,7 +941,7 @@ bool UDataLayerEditorSubsystem::ResetUserSettings()
 	bool bRefreshNeeded = false;
 	if (const AWorldDataLayers* WorldDataLayers = GetWorldDataLayers())
 	{
-		WorldDataLayers->ForEachDataLayer([this, &bRefreshNeeded](UDataLayer* DataLayer)
+		WorldDataLayers->ForEachDataLayer([this, &bRefreshNeeded](UDataLayerInstance* DataLayer)
 		{
 			bRefreshNeeded |= SetDataLayerIsLoadedInEditorInternal(DataLayer, DataLayer->IsInitiallyLoadedInEditor(), true);
 			return true;
@@ -943,27 +950,26 @@ bool UDataLayerEditorSubsystem::ResetUserSettings()
 	return bRefreshNeeded ? RefreshWorldPartitionEditorCells(true) : true;
 }
 
-UDataLayer* UDataLayerEditorSubsystem::GetDataLayer(const FActorDataLayer& ActorDataLayer) const
+bool UDataLayerEditorSubsystem::HasDeprecatedDataLayers() const
 {
-	return GetDataLayerFromName(ActorDataLayer.Name);
-}
+	if (const AWorldDataLayers* WorldDataLayers = GetWorldDataLayers())
+	{
+		return WorldDataLayers->HasDeprecatedDataLayers();
+	}
 
-UDataLayer* UDataLayerEditorSubsystem::GetDataLayerFromName(const FName& DataLayerName) const
-{
-	const AWorldDataLayers* WorldDataLayers = GetWorldDataLayers();
-	return WorldDataLayers ? const_cast<UDataLayer*>(WorldDataLayers->GetDataLayerFromName(DataLayerName)) : nullptr;
+	return false;
 }
 
-UDataLayer* UDataLayerEditorSubsystem::GetDataLayerFromLabel(const FName& DataLayerLabel) const
+UDataLayerInstance* UDataLayerEditorSubsystem::GetDataLayerInstance(const FName& DataLayerInstanceName) const
 {
 	const AWorldDataLayers* WorldDataLayers = GetWorldDataLayers();
-	return WorldDataLayers ? const_cast<UDataLayer*>(WorldDataLayers->GetDataLayerFromLabel(DataLayerLabel)) : nullptr;
+	return WorldDataLayers ? const_cast<UDataLayerInstance*>(WorldDataLayers->GetDataLayerInstance(DataLayerInstanceName)) : nullptr;
 }
 
-bool UDataLayerEditorSubsystem::TryGetDataLayerFromLabel(const FName& DataLayerLabel, UDataLayer*& OutDataLayer)
+UDataLayerInstance* UDataLayerEditorSubsystem::GetDataLayerInstance(const UDataLayerAsset* DataLayerAsset) const
 {
-	OutDataLayer = GetDataLayerFromLabel(DataLayerLabel);
-	return (OutDataLayer != nullptr);
+	const AWorldDataLayers* WorldDataLayers = GetWorldDataLayers();
+	return WorldDataLayers ? const_cast<UDataLayerInstance*>(WorldDataLayers->GetDataLayerInstance(DataLayerAsset)) : nullptr;
 }
 
 const AWorldDataLayers* UDataLayerEditorSubsystem::GetWorldDataLayers() const
@@ -983,27 +989,67 @@ AWorldDataLayers* UDataLayerEditorSubsystem::GetWorldDataLayers(bool bCreateIfNo
 	return WorldDataLayers;
 }
 
-void UDataLayerEditorSubsystem::AddAllDataLayersTo(TArray<TWeakObjectPtr<UDataLayer>>& OutDataLayers) const
+void UDataLayerEditorSubsystem::AddAllDataLayersTo(TArray<TWeakObjectPtr<UDataLayerInstance>>& OutDataLayers) const
 {
 	if (const AWorldDataLayers* WorldDataLayers = GetWorldDataLayers())
 	{
-		WorldDataLayers->ForEachDataLayer([&OutDataLayers](UDataLayer* DataLayer)
-		{
-			OutDataLayers.Add(DataLayer);
-			return true;
-		});
+		WorldDataLayers->ForEachDataLayer([&OutDataLayers](UDataLayerInstance* DataLayer)
+			{
+				OutDataLayers.Add(DataLayer);
+				return true;
+			});
 	}
 }
 
-UDataLayer* UDataLayerEditorSubsystem::CreateDataLayer()
+UDataLayerInstance* UDataLayerEditorSubsystem::CreateDataLayer(UDataLayerInstance* ParentDataLayer)
 {
+	UDataLayerInstance* NewDataLayer = nullptr;
+
+
 	AWorldDataLayers* WorldDataLayers = GetWorldDataLayers(/*bCreateIfNotFound*/true);
-	UDataLayer* NewDataLayer = WorldDataLayers->CreateDataLayer();
-	BroadcastDataLayerChanged(EDataLayerAction::Add, NewDataLayer, NAME_None);
+	if (!WorldDataLayers->HasDeprecatedDataLayers())
+	{
+		bool bContinuePrompt = true;
+		while (bContinuePrompt)
+		{
+			FAssetData AssetData;
+			bContinuePrompt = PromptDataLayerAssetSelection(AssetData);
+			if (bContinuePrompt)
+			{
+				UDataLayerAsset* DataLayerAsset = CastChecked<UDataLayerAsset>(AssetData.GetAsset());
+				if (UDataLayerInstance* DataLayerInstance = GetDataLayerInstance(DataLayerAsset))
+				{
+					FText WarningMessage(LOCTEXT("CreateDataLayerAlreadyExist", "A Data Layer of that asset already exists. Pick another Data Layer?"));
+					bContinuePrompt = FMessageDialog::Open(EAppMsgType::YesNo, WarningMessage) == EAppReturnType::Yes;
+				}
+				else if (ParentDataLayer != nullptr && ParentDataLayer->GetType() != DataLayerAsset->GetType())
+				{
+					FText WarningMessage(LOCTEXT("CreateDataLayerDiffType", "The parent data layer is of a different type. Pick another Data Layer?"));
+					bContinuePrompt = FMessageDialog::Open(EAppMsgType::YesNo, WarningMessage) == EAppReturnType::Yes;
+				}
+				else
+				{
+					NewDataLayer = WorldDataLayers->CreateDataLayer<UDataLayerInstanceWithAsset>(DataLayerAsset);
+					bContinuePrompt = false;
+				}
+			}
+		}
+	}
+	else
+	{
+		NewDataLayer = WorldDataLayers->CreateDataLayer<UDeprecatedDataLayerInstance>();
+	}
+
+	if (NewDataLayer != nullptr)
+	{
+		BroadcastDataLayerChanged(EDataLayerAction::Add, NewDataLayer, NAME_None);
+		SetParentDataLayer(NewDataLayer, ParentDataLayer);
+	}
+	
 	return NewDataLayer;
 }
 
-void UDataLayerEditorSubsystem::DeleteDataLayers(const TArray<UDataLayer*>& DataLayersToDelete)
+void UDataLayerEditorSubsystem::DeleteDataLayers(const TArray<UDataLayerInstance*>& DataLayersToDelete)
 {
 	if (AWorldDataLayers* WorldDataLayers = GetWorldDataLayers())
 	{
@@ -1014,7 +1060,7 @@ void UDataLayerEditorSubsystem::DeleteDataLayers(const TArray<UDataLayer*>& Data
 	}
 }
 
-void UDataLayerEditorSubsystem::DeleteDataLayer(UDataLayer* DataLayerToDelete)
+void UDataLayerEditorSubsystem::DeleteDataLayer(UDataLayerInstance* DataLayerToDelete)
 {
 	if (AWorldDataLayers* WorldDataLayers = GetWorldDataLayers())
 	{
@@ -1025,30 +1071,13 @@ void UDataLayerEditorSubsystem::DeleteDataLayer(UDataLayer* DataLayerToDelete)
 	}
 }
 
-bool UDataLayerEditorSubsystem::RenameDataLayer(UDataLayer* DataLayer, const FName& InDataLayerLabel)
-{
-	const FName DataLayerLabelSanitized = UDataLayer::GetSanitizedDataLayerLabel(InDataLayerLabel);
-	if (DataLayer->GetDataLayerLabel() != DataLayerLabelSanitized)
-	{
-		if (AWorldDataLayers* WorldDataLayers = GetWorldDataLayers())
-		{
-			FName UniqueNewDataLayerLabel = WorldDataLayers->GenerateUniqueDataLayerLabel(DataLayerLabelSanitized);
-			DataLayer->Modify();
-			DataLayer->SetDataLayerLabel(UniqueNewDataLayerLabel);
-			BroadcastDataLayerChanged(EDataLayerAction::Rename, DataLayer, "DataLayerLabel");
-			return true;
-		}
-	}
-	return false;
-}
-
 void UDataLayerEditorSubsystem::BroadcastActorDataLayersChanged(const TWeakObjectPtr<AActor>& ChangedActor)
 {
 	RebuildSelectedDataLayersFromEditorSelection();
 	ActorDataLayersChanged.Broadcast(ChangedActor);
 }
 
-void UDataLayerEditorSubsystem::BroadcastDataLayerChanged(const EDataLayerAction Action, const TWeakObjectPtr<const UDataLayer>& ChangedDataLayer, const FName& ChangedProperty)
+void UDataLayerEditorSubsystem::BroadcastDataLayerChanged(const EDataLayerAction Action, const TWeakObjectPtr<const UDataLayerInstance>& ChangedDataLayer, const FName& ChangedProperty)
 {
 	RebuildSelectedDataLayersFromEditorSelection();
 	DataLayerChanged.Broadcast(Action, ChangedDataLayer, ChangedProperty);
@@ -1068,9 +1097,69 @@ void UDataLayerEditorSubsystem::RebuildSelectedDataLayersFromEditorSelection()
 	GEditor->GetSelectedActors()->GetSelectedObjects<AActor>(Actors);
 	for (const AActor* Actor : Actors)
 	{
-		for (const UDataLayer* DataLayer : Actor->GetDataLayerObjects())
+		for (const UDataLayerInstance* DataLayerInstance : Actor->GetDataLayerInstances())
 		{
-			SelectedDataLayersFromEditorSelection.Add(DataLayer);
+			SelectedDataLayersFromEditorSelection.Add(DataLayerInstance);
 		}
 	}
 }
+
+
+bool UDataLayerEditorSubsystem::PromptDataLayerAssetSelection(FAssetData& OutAsset) const
+{
+	IContentBrowserSingleton& ContentBrowserSingleton = FModuleManager::LoadModuleChecked<FContentBrowserModule>("ContentBrowser").Get();
+
+	FOpenAssetDialogConfig Config;
+	Config.bAllowMultipleSelection = false;
+	Config.AssetClassNames.Add(UDataLayerAsset::StaticClass()->GetFName());
+
+	TArray<FAssetData> Assets = ContentBrowserSingleton.CreateModalOpenAssetDialog(Config);
+	if (Assets.Num() == 1)
+	{
+		OutAsset = Assets[0];
+		return true;
+	}
+
+	return false;
+}
+
+//~ Begin Deprecated
+
+PRAGMA_DISABLE_DEPRECATION_WARNINGS
+
+bool UDataLayerEditorSubsystem::RenameDataLayer(UDataLayerInstance* DataLayerInstance, const FName& InDataLayerLabel)
+{
+	if (DataLayerInstance->SupportRelabeling())
+	{
+		if (DataLayerInstance->RelabelDataLayer(InDataLayerLabel))
+		{
+			BroadcastDataLayerChanged(EDataLayerAction::Rename, DataLayerInstance, "DataLayerLabel");
+			return true;
+		}	
+	}
+
+	return false;
+}
+
+bool UDataLayerEditorSubsystem::TryGetDataLayerFromLabel(const FName& DataLayerLabel, UDataLayerInstance*& OutDataLayer) const
+{
+	OutDataLayer = GetDataLayerFromLabel(DataLayerLabel);
+	return (OutDataLayer != nullptr);
+}
+
+UDataLayerInstance* UDataLayerEditorSubsystem::GetDataLayerFromLabel(const FName& DataLayerLabel) const
+{
+	const AWorldDataLayers* WorldDataLayers = GetWorldDataLayers();
+	return WorldDataLayers ? const_cast<UDataLayerInstance*>(WorldDataLayers->GetDataLayerFromLabel(DataLayerLabel)) : nullptr;
+}
+
+PRAGMA_ENABLE_DEPRECATION_WARNINGS
+
+UDataLayerInstance* UDataLayerEditorSubsystem::GetDataLayer(const FActorDataLayer& ActorDataLayer) const
+{
+	return GetDataLayerInstance(ActorDataLayer.Name);
+}
+
+//~ End Deprecated
+
+#undef LOCTEXT_NAMESPACE
\ No newline at end of file
diff --git a/Engine/Source/Editor/DataLayerEditor/Private/DataLayer/DataLayerHierarchy.cpp b/Engine/Source/Editor/DataLayerEditor/Private/DataLayer/DataLayerHierarchy.cpp
--- a/Engine/Source/Editor/DataLayerEditor/Private/DataLayer/DataLayerHierarchy.cpp
+++ b/Engine/Source/Editor/DataLayerEditor/Private/DataLayer/DataLayerHierarchy.cpp
@@ -7,7 +7,7 @@
 #include "DataLayerTreeItem.h"
 #include "DataLayerMode.h"
 #include "DataLayer/DataLayerEditorSubsystem.h"
-#include "WorldPartition/DataLayer/DataLayer.h"
+#include "WorldPartition/DataLayer/DataLayerAsset.h"
 #include "WorldPartition/DataLayer/WorldDataLayers.h"
 #include "Engine/Engine.h"
 #include "Engine/World.h"
@@ -97,7 +97,7 @@ FDataLayerHierarchy::~FDataLayerHierarchy()
 }
 
 
-bool FDataLayerHierarchy::IsDataLayerPartOfSelection(const UDataLayer* DataLayer) const
+bool FDataLayerHierarchy::IsDataLayerPartOfSelection(const UDataLayerInstance* DataLayer) const
 {
 	if (!bShowOnlySelectedActors)
 	{
@@ -110,7 +110,7 @@ bool FDataLayerHierarchy::IsDataLayerPartOfSelection(const UDataLayer* DataLayer
 	}
 
 	bool bFoundSelected = false;
-	DataLayer->ForEachChild([this, &bFoundSelected](const UDataLayer* Child)
+	DataLayer->ForEachChild([this, &bFoundSelected](const UDataLayerInstance* Child)
 	{
 		bFoundSelected = IsDataLayerPartOfSelection(Child);
 		return !bFoundSelected; // Continue iterating if not found
@@ -118,7 +118,7 @@ bool FDataLayerHierarchy::IsDataLayerPartOfSelection(const UDataLayer* DataLayer
 	return bFoundSelected;
 };
 
-FSceneOutlinerTreeItemPtr FDataLayerHierarchy::CreateDataLayerTreeItem(UDataLayer* InDataLayer, bool bInForce) const
+FSceneOutlinerTreeItemPtr FDataLayerHierarchy::CreateDataLayerTreeItem(UDataLayerInstance* InDataLayer, bool bInForce) const
 {
 	FSceneOutlinerTreeItemPtr Item = Mode->CreateItemFor<FDataLayerTreeItem>(InDataLayer, bInForce);
 	if (FDataLayerTreeItem* DataLayerTreeItem = Item ? Item->CastTo<FDataLayerTreeItem>() : nullptr)
@@ -136,13 +136,13 @@ void FDataLayerHierarchy::CreateItems(TArray<FSceneOutlinerTreeItemPtr>& OutItem
 		return;
 	}
 
-	auto IsDataLayerShown = [this](const UDataLayer* DataLayer)
+	auto IsDataLayerShown = [this](const UDataLayerInstance* DataLayer)
 	{
 		const bool bIsRuntimeDataLayer = DataLayer->IsRuntime();
 		return ((bIsRuntimeDataLayer && bShowRuntimeDataLayers) || (!bIsRuntimeDataLayer && bShowEditorDataLayers)) && IsDataLayerPartOfSelection(DataLayer);
 	};
 
-	WorldDataLayers->ForEachDataLayer([this, &OutItems, IsDataLayerShown](UDataLayer* DataLayer)
+	WorldDataLayers->ForEachDataLayer([this, &OutItems, IsDataLayerShown](UDataLayerInstance* DataLayer)
 	{
 		if (IsDataLayerShown(DataLayer))
 		{
@@ -160,11 +160,11 @@ void FDataLayerHierarchy::CreateItems(TArray<FSceneOutlinerTreeItemPtr>& OutItem
 		{
 			if (Actor->HasDataLayers())
 			{
-				for (const UDataLayer* DataLayer : Actor->GetDataLayerObjects())
+				for (const UDataLayerInstance* DataLayerInstance : Actor->GetDataLayerInstances())
 				{
-					if (IsDataLayerShown(DataLayer))
+					if (IsDataLayerShown(DataLayerInstance))
 					{
-						if (FSceneOutlinerTreeItemPtr DataLayerActorItem = Mode->CreateItemFor<FDataLayerActorTreeItem>(FDataLayerActorTreeItemData(Actor, const_cast<UDataLayer*>(DataLayer))))
+						if (FSceneOutlinerTreeItemPtr DataLayerActorItem = Mode->CreateItemFor<FDataLayerActorTreeItem>(FDataLayerActorTreeItemData(Actor, const_cast<UDataLayerInstance*>(DataLayerInstance))))
 						{
 							OutItems.Add(DataLayerActorItem);
 						}
@@ -181,13 +181,13 @@ void FDataLayerHierarchy::CreateItems(TArray<FSceneOutlinerTreeItemPtr>& OutItem
 				{
 					if (ActorDesc != nullptr && !ActorDesc->IsLoaded())
 					{
-						for (const FName& DataLayerName : ActorDesc->GetDataLayers())
+						for (const FName& DataLayerInstanceName : ActorDesc->GetDataLayerInstanceNames())
 						{
-							if (const UDataLayer* const DataLayer = WorldDataLayers->GetDataLayerFromName(DataLayerName))
+							if (const UDataLayerInstance* const DataLayerInstance = WorldDataLayers->GetDataLayerInstance(DataLayerInstanceName))
 							{
-								if (IsDataLayerShown(DataLayer))
+								if (IsDataLayerShown(DataLayerInstance))
 								{
-									if (const FSceneOutlinerTreeItemPtr ActorDescItem = Mode->CreateItemFor<FDataLayerActorDescTreeItem>(FDataLayerActorDescTreeItemData(ActorDesc->GetGuid(), WorldPartition, const_cast<UDataLayer*>(DataLayer))))
+									if (const FSceneOutlinerTreeItemPtr ActorDescItem = Mode->CreateItemFor<FDataLayerActorDescTreeItem>(FDataLayerActorDescTreeItemData(ActorDesc->GetGuid(), WorldPartition, const_cast<UDataLayerInstance*>(DataLayerInstance))))
 									{
 										OutItems.Add(ActorDescItem);
 									}
@@ -206,9 +206,9 @@ FSceneOutlinerTreeItemPtr FDataLayerHierarchy::FindOrCreateParentItem(const ISce
 {
 	if (const FDataLayerTreeItem* DataLayerTreeItem = Item.CastTo<FDataLayerTreeItem>())
 	{
-		if (UDataLayer* DataLayer = DataLayerTreeItem->GetDataLayer())
+		if (UDataLayerInstance* DataLayer = DataLayerTreeItem->GetDataLayer())
 		{
-			if (UDataLayer* ParentDataLayer = DataLayer->GetParent())
+			if (UDataLayerInstance* ParentDataLayer = DataLayer->GetParent())
 			{
 				if (const FSceneOutlinerTreeItemPtr* ParentItem = Items.Find(ParentDataLayer))
 				{
@@ -223,7 +223,7 @@ FSceneOutlinerTreeItemPtr FDataLayerHierarchy::FindOrCreateParentItem(const ISce
 	}
 	else if (const FDataLayerActorTreeItem* DataLayerActorTreeItem = Item.CastTo<FDataLayerActorTreeItem>())
 	{
-		if (UDataLayer* DataLayer = DataLayerActorTreeItem->GetDataLayer())
+		if (UDataLayerInstance* DataLayer = DataLayerActorTreeItem->GetDataLayer())
 		{
 			if (const FSceneOutlinerTreeItemPtr* ParentItem = Items.Find(DataLayer))
 			{
@@ -237,7 +237,7 @@ FSceneOutlinerTreeItemPtr FDataLayerHierarchy::FindOrCreateParentItem(const ISce
 	}
 	else if (const FDataLayerActorDescTreeItem* DataLayerActorDescTreeItem = Item.CastTo<FDataLayerActorDescTreeItem>())
 	{
-		if (UDataLayer* DataLayer = DataLayerActorDescTreeItem->GetDataLayer())
+		if (UDataLayerInstance* DataLayer = DataLayerActorDescTreeItem->GetDataLayer())
 		{
 			if (const FSceneOutlinerTreeItemPtr* ParentItem = Items.Find(DataLayer))
 			{
@@ -276,11 +276,11 @@ void FDataLayerHierarchy::OnLevelActorsAdded(const TArray<AActor*>& InActors)
 		{
 			if (Actor != nullptr && Actor->HasDataLayers() && Actor->GetWorld() == CurrentWorld)
 			{
-				TArray<const UDataLayer*> DataLayers = Actor->GetDataLayerObjects();
-				EventData.Items.Reserve(DataLayers.Num());
-				for (const UDataLayer* DataLayer : DataLayers)
+				TArray<const UDataLayerInstance*> DataLayerInstances = Actor->GetDataLayerInstances();
+				EventData.Items.Reserve(DataLayerInstances.Num());
+				for (const UDataLayerInstance* DataLayerInstance : DataLayerInstances)
 				{
-					EventData.Items.Add(Mode->CreateItemFor<FDataLayerActorTreeItem>(FDataLayerActorTreeItemData(Actor, const_cast<UDataLayer*>(DataLayer))));
+					EventData.Items.Add(Mode->CreateItemFor<FDataLayerActorTreeItem>(FDataLayerActorTreeItemData(Actor, const_cast<UDataLayerInstance*>(DataLayerInstance))));
 				}
 			}
 		}
@@ -306,11 +306,11 @@ void FDataLayerHierarchy::OnLevelActorsRemoved(const TArray<AActor*>& InActors)
 				if (Actor != nullptr && Actor->HasDataLayers())
 				{
 					// It is possible here that Actor doesn't have world anymore
-					const TArray<const UDataLayer*> DataLayers = Actor->GetDataLayerObjects(WorldDataLayers);
-					EventData.ItemIDs.Reserve(DataLayers.Num());
-					for (const UDataLayer* DataLayer : DataLayers)
+					const TArray<const UDataLayerInstance*> DataLayerInstances = Actor->GetDataLayerInstances(WorldDataLayers);
+					EventData.ItemIDs.Reserve(DataLayerInstances.Num());
+					for (const UDataLayerInstance* DataLayerInstance : DataLayerInstances)
 					{
-						EventData.ItemIDs.Add(FDataLayerActorTreeItem::ComputeTreeItemID(Actor, DataLayer));
+						EventData.ItemIDs.Add(FDataLayerActorTreeItem::ComputeTreeItemID(Actor, DataLayerInstance));
 					}
 				}
 			}
@@ -337,9 +337,9 @@ void FDataLayerHierarchy::OnActorDataLayersChanged(const TWeakObjectPtr<AActor>&
 	}
 }
 
-void FDataLayerHierarchy::OnDataLayerChanged(const EDataLayerAction Action, const TWeakObjectPtr<const UDataLayer>& ChangedDataLayer, const FName& ChangedProperty)
+void FDataLayerHierarchy::OnDataLayerChanged(const EDataLayerAction Action, const TWeakObjectPtr<const UDataLayerInstance>& ChangedDataLayer, const FName& ChangedProperty)
 {
-	const UDataLayer* DataLayer = ChangedDataLayer.Get();
+	const UDataLayerInstance* DataLayer = ChangedDataLayer.Get();
 	if ((DataLayer && (RepresentingWorld.Get() == DataLayer->GetWorld())) || (Action == EDataLayerAction::Delete) || (Action == EDataLayerAction::Reset))
 	{
 		FullRefreshEvent();
@@ -388,11 +388,11 @@ void FDataLayerHierarchy::OnLoadedActorAdded(AActor& InActor)
 		FSceneOutlinerHierarchyChangedData EventData;
 		EventData.Type = FSceneOutlinerHierarchyChangedData::Removed;
 		
-		const TArray<const UDataLayer*> DataLayers = InActor.GetDataLayerObjects();
-		EventData.ItemIDs.Reserve(DataLayers.Num());
-		for (const UDataLayer* DataLayer : DataLayers)
+		const TArray<const UDataLayerInstance*> DataLayerInstances = InActor.GetDataLayerInstances();
+		EventData.ItemIDs.Reserve(DataLayerInstances.Num());
+		for (const UDataLayerInstance* DataLayerInstance : DataLayerInstances)
 		{
-			EventData.ItemIDs.Add(FDataLayerActorDescTreeItem::ComputeTreeItemID(InActor.GetActorGuid(), DataLayer));
+			EventData.ItemIDs.Add(FDataLayerActorDescTreeItem::ComputeTreeItemID(InActor.GetActorGuid(), DataLayerInstance));
 		}
 		HierarchyChangedEvent.Broadcast(EventData);
 	}
@@ -410,11 +410,11 @@ void FDataLayerHierarchy::OnLoadedActorRemoved(AActor& InActor)
 		FSceneOutlinerHierarchyChangedData EventData;
 		EventData.Type = FSceneOutlinerHierarchyChangedData::Added;
 
-		const TArray<const UDataLayer*> DataLayers = InActor.GetDataLayerObjects();
-		EventData.Items.Reserve(DataLayers.Num());
-		for (const UDataLayer* DataLayer : DataLayers)
+		const TArray<const UDataLayerInstance*> DataLayerInstances = InActor.GetDataLayerInstances();
+		EventData.Items.Reserve(DataLayerInstances.Num());
+		for (const UDataLayerInstance* DataLayerInstance : DataLayerInstances)
 		{
-			EventData.Items.Add(Mode->CreateItemFor<FDataLayerActorDescTreeItem>(FDataLayerActorDescTreeItemData(InActor.GetActorGuid(), WorldPartition, const_cast<UDataLayer*>(DataLayer))));
+			EventData.Items.Add(Mode->CreateItemFor<FDataLayerActorDescTreeItem>(FDataLayerActorDescTreeItemData(InActor.GetActorGuid(), WorldPartition, const_cast<UDataLayerInstance*>(DataLayerInstance))));
 		}
 		HierarchyChangedEvent.Broadcast(EventData);
     }
@@ -429,18 +429,18 @@ void FDataLayerHierarchy::OnActorDescAdded(FWorldPartitionActorDesc* InActorDesc
 
 	UWorldPartition* const WorldPartition = RepresentingWorld.Get()->GetWorldPartition();
 	const AWorldDataLayers* const WorldDataLayers = RepresentingWorld.Get()->GetWorldDataLayers();
-	const TArray<FName>& DataLayerNames = InActorDesc->GetDataLayers();
+	const TArray<FName>& DataLayerInstanceNames = InActorDesc->GetDataLayerInstanceNames();
 
-	if (WorldDataLayers != nullptr && RepresentingWorld->GetWorldPartition() == WorldPartition && DataLayerNames.Num() > 0)
+	if (WorldDataLayers != nullptr && RepresentingWorld->GetWorldPartition() == WorldPartition && DataLayerInstanceNames.Num() > 0)
 	{
 		FSceneOutlinerHierarchyChangedData EventData;
 		EventData.Type = FSceneOutlinerHierarchyChangedData::Added;
 
-		EventData.Items.Reserve(DataLayerNames.Num());
-		for (const FName& DataLayerName : DataLayerNames)
+		EventData.Items.Reserve(DataLayerInstanceNames.Num());
+		for (const FName& DataLayerInstanceName : DataLayerInstanceNames)
 		{
-			const UDataLayer* const DataLayer = WorldDataLayers->GetDataLayerFromName(DataLayerName);
-			EventData.Items.Add(Mode->CreateItemFor<FDataLayerActorDescTreeItem>(FDataLayerActorDescTreeItemData(InActorDesc->GetGuid(), WorldPartition, const_cast<UDataLayer*>(DataLayer))));
+			const UDataLayerInstance* const DataLayerInstance = WorldDataLayers->GetDataLayerInstance(DataLayerInstanceName);
+			EventData.Items.Add(Mode->CreateItemFor<FDataLayerActorDescTreeItem>(FDataLayerActorDescTreeItemData(InActorDesc->GetGuid(), WorldPartition, const_cast<UDataLayerInstance*>(DataLayerInstance))));
 		}
 		HierarchyChangedEvent.Broadcast(EventData);
 	}
@@ -454,17 +454,17 @@ void FDataLayerHierarchy::OnActorDescRemoved(FWorldPartitionActorDesc* InActorDe
 	}
 
 	const AWorldDataLayers* const WorldDataLayers = RepresentingWorld.Get()->GetWorldDataLayers();
-	const TArray<FName>& DataLayerNames = InActorDesc->GetDataLayers();
+	const TArray<FName>& DataLayerInstanceNames = InActorDesc->GetDataLayerInstanceNames();
 	
-	if (WorldDataLayers != nullptr && DataLayerNames.Num() > 0)
+	if (WorldDataLayers != nullptr && DataLayerInstanceNames.Num() > 0)
 	{
 		FSceneOutlinerHierarchyChangedData EventData;
 		EventData.Type = FSceneOutlinerHierarchyChangedData::Removed;
-		EventData.ItemIDs.Reserve(DataLayerNames.Num());
+		EventData.ItemIDs.Reserve(DataLayerInstanceNames.Num());
 
-		for (const FName& DataLayerName : DataLayerNames)
+		for (const FName& DataLayerInstanceName : DataLayerInstanceNames)
 		{
-			const UDataLayer* const DataLayer = WorldDataLayers->GetDataLayerFromName(DataLayerName);
+			const UDataLayerInstance* const DataLayer = WorldDataLayers->GetDataLayerInstance(DataLayerInstanceName);
 			EventData.ItemIDs.Add(FDataLayerActorDescTreeItem::ComputeTreeItemID(InActorDesc->GetGuid(), DataLayer));
 		}
 		HierarchyChangedEvent.Broadcast(EventData);
diff --git a/Engine/Source/Editor/DataLayerEditor/Private/DataLayer/DataLayerHierarchy.h b/Engine/Source/Editor/DataLayerEditor/Private/DataLayer/DataLayerHierarchy.h
--- a/Engine/Source/Editor/DataLayerEditor/Private/DataLayer/DataLayerHierarchy.h
+++ b/Engine/Source/Editor/DataLayerEditor/Private/DataLayer/DataLayerHierarchy.h
@@ -8,7 +8,7 @@
 #include "DataLayer/DataLayerAction.h"
 
 class FDataLayerMode;
-class UDataLayer;
+class UDataLayerInstance;
 class UWorld;
 class FWorldPartitionActorDesc;
 
@@ -45,10 +45,10 @@ private:
 	void OnActorDescAdded(FWorldPartitionActorDesc* InActorDesc);
 	void OnActorDescRemoved(FWorldPartitionActorDesc* InActorDesc);
 	void OnActorDataLayersChanged(const TWeakObjectPtr<AActor>& InActor);
-	void OnDataLayerChanged(const EDataLayerAction Action, const TWeakObjectPtr<const UDataLayer>& ChangedDataLayer, const FName& ChangedProperty);
+	void OnDataLayerChanged(const EDataLayerAction Action, const TWeakObjectPtr<const UDataLayerInstance>& ChangedDataLayer, const FName& ChangedProperty);
 	void FullRefreshEvent();
-	FSceneOutlinerTreeItemPtr CreateDataLayerTreeItem(UDataLayer* InDataLayer, bool bInForce = false) const;
-	bool IsDataLayerPartOfSelection(const UDataLayer* DataLayer) const;
+	FSceneOutlinerTreeItemPtr CreateDataLayerTreeItem(UDataLayerInstance* InDataLayer, bool bInForce = false) const;
+	bool IsDataLayerPartOfSelection(const UDataLayerInstance* DataLayer) const;
 
 	TWeakObjectPtr<UWorld> RepresentingWorld;
 	bool bShowEditorDataLayers;
