UDataLayer deprecated in favor of UDataLayerInstance & UDataLayerAsset

diff --git a/Engine/Source/Editor/DataLayerEditor/Private/DataLayer/DataLayerMode.cpp b/Engine/Source/Editor/DataLayerEditor/Private/DataLayer/DataLayerMode.cpp
--- a/Engine/Source/Editor/DataLayerEditor/Private/DataLayer/DataLayerMode.cpp
+++ b/Engine/Source/Editor/DataLayerEditor/Private/DataLayer/DataLayerMode.cpp
@@ -130,7 +130,7 @@ FDataLayerMode::FDataLayerMode(const FDataLayerModeParams& Params)
 	});
 	FilterInfoMap.Add(TEXT("HideUnloadedActorsFilter"), HideUnloadedActorsInfo);
 
-	SceneOutliner->AddFilter(MakeShared<TSceneOutlinerPredicateFilter<FDataLayerActorTreeItem>>(FDataLayerActorTreeItem::FFilterPredicate::CreateLambda([this](const AActor* Actor, const UDataLayer* DataLayer)
+	SceneOutliner->AddFilter(MakeShared<TSceneOutlinerPredicateFilter<FDataLayerActorTreeItem>>(FDataLayerActorTreeItem::FFilterPredicate::CreateLambda([this](const AActor* Actor, const UDataLayerInstance* DataLayer)
 	{
 		return FActorMode::IsActorDisplayable(SceneOutliner, Actor);
 	}), FSceneOutlinerFilter::EDefaultBehaviour::Pass));
@@ -148,17 +148,17 @@ FDataLayerMode::~FDataLayerMode()
 
 TSharedRef<FSceneOutlinerFilter> FDataLayerMode::CreateHideEditorDataLayersFilter()
 {
-	return MakeShareable(new FDataLayerFilter(FDataLayerTreeItem::FFilterPredicate::CreateStatic([](const UDataLayer* DataLayer) { return true; }), FSceneOutlinerFilter::EDefaultBehaviour::Pass));
+	return MakeShareable(new FDataLayerFilter(FDataLayerTreeItem::FFilterPredicate::CreateStatic([](const UDataLayerInstance* DataLayerInstance) { return true; }), FSceneOutlinerFilter::EDefaultBehaviour::Pass));
 }
 
 TSharedRef<FSceneOutlinerFilter> FDataLayerMode::CreateHideRuntimeDataLayersFilter()
 {
-	return MakeShareable(new FDataLayerFilter(FDataLayerTreeItem::FFilterPredicate::CreateStatic([](const UDataLayer* DataLayer) { return true; }), FSceneOutlinerFilter::EDefaultBehaviour::Pass));
+	return MakeShareable(new FDataLayerFilter(FDataLayerTreeItem::FFilterPredicate::CreateStatic([](const UDataLayerInstance* DataLayerInstance) { return true; }), FSceneOutlinerFilter::EDefaultBehaviour::Pass));
 }
 
 TSharedRef<FSceneOutlinerFilter> FDataLayerMode::CreateHideDataLayerActorsFilter()
 {
-	return MakeShareable(new FDataLayerActorFilter(FDataLayerActorTreeItem::FFilterPredicate::CreateStatic([](const AActor* Actor, const UDataLayer* DataLayer) { return true; }), FSceneOutlinerFilter::EDefaultBehaviour::Pass));
+	return MakeShareable(new FDataLayerActorFilter(FDataLayerActorTreeItem::FFilterPredicate::CreateStatic([](const AActor* Actor, const UDataLayerInstance* DataDataLayerInstanceLayer) { return true; }), FSceneOutlinerFilter::EDefaultBehaviour::Pass));
 }
 
 TSharedRef<FSceneOutlinerFilter> FDataLayerMode::CreateHideUnloadedActorsFilter()
@@ -190,7 +190,7 @@ bool FDataLayerMode::CanRenameItem(const ISceneOutlinerTreeItem& Item) const
 	if (Item.IsValid() && (Item.IsA<FDataLayerTreeItem>()))
 	{
 		FDataLayerTreeItem* DataLayerTreeItem = (FDataLayerTreeItem*)&Item;
-		return !DataLayerTreeItem->GetDataLayer()->IsLocked();
+		return !DataLayerTreeItem->GetDataLayer()->IsLocked() && DataLayerTreeItem->GetDataLayer()->SupportRelabeling();
 	}
 	return false;
 }
@@ -282,11 +282,11 @@ void FDataLayerMode::OnItemDoubleClick(FSceneOutlinerTreeItemPtr Item)
 {
 	if (const FDataLayerTreeItem* DataLayerItem = Item->CastTo<FDataLayerTreeItem>())
 	{
-		if (UDataLayer* DataLayer = DataLayerItem->GetDataLayer())
+		if (UDataLayerInstance* DataLayerInstance = DataLayerItem->GetDataLayer())
 		{
 			const FScopedDataLayerTransaction Transaction(LOCTEXT("SelectActorsInDataLayer", "Select Actors in Data Layer"), RepresentingWorld.Get());
 			GEditor->SelectNone(/*bNoteSelectionChange*/false, true);
-			DataLayerEditorSubsystem->SelectActorsInDataLayer(DataLayer, /*bSelect*/true, /*bNotify*/true, /*bSelectEvenIfHidden*/true);
+			DataLayerEditorSubsystem->SelectActorsInDataLayer(DataLayerInstance, /*bSelect*/true, /*bNotify*/true, /*bSelectEvenIfHidden*/true);
 		}
 	}
 	else if (FDataLayerActorTreeItem* DataLayerActorItem = Item->CastTo<FDataLayerActorTreeItem>())
@@ -305,27 +305,27 @@ void FDataLayerMode::OnItemDoubleClick(FSceneOutlinerTreeItemPtr Item)
 
 void FDataLayerMode::DeleteItems(const TArray<TWeakPtr<ISceneOutlinerTreeItem>>& Items)
 {
-	TArray<UDataLayer*> DataLayersToDelete;
-	TMap<UDataLayer*, TArray<AActor*>> ActorsToRemoveFromDataLayer;
+	TArray<UDataLayerInstance*> DataLayersToDelete;
+	TMap<UDataLayerInstance*, TArray<AActor*>> ActorsToRemoveFromDataLayer;
 
 	for (const TWeakPtr<ISceneOutlinerTreeItem>& Item : Items)
 	{
 		if (FDataLayerActorTreeItem* DataLayerActorItem = Item.Pin()->CastTo<FDataLayerActorTreeItem>())
 		{
-			UDataLayer* DataLayer = DataLayerActorItem->GetDataLayer();
+			UDataLayerInstance* DataLayerInstance = DataLayerActorItem->GetDataLayer();
 			AActor* Actor = DataLayerActorItem->GetActor();
-			if (DataLayer && !DataLayer->IsLocked() && Actor)
+			if (DataLayerInstance && !DataLayerInstance->IsLocked() && Actor)
 			{
-				ActorsToRemoveFromDataLayer.FindOrAdd(DataLayer).Add(Actor);
+				ActorsToRemoveFromDataLayer.FindOrAdd(DataLayerInstance).Add(Actor);
 			}
 		}
 		else if (FDataLayerTreeItem* DataLayerItem = Item.Pin()->CastTo< FDataLayerTreeItem>())
 		{
-			if (UDataLayer* DataLayer = DataLayerItem->GetDataLayer())
+			if (UDataLayerInstance* DataLayerInstance = DataLayerItem->GetDataLayer())
 			{
-				if (!DataLayer->IsLocked())
+				if (!DataLayerInstance->IsLocked())
 				{
-					DataLayersToDelete.Add(DataLayer);
+					DataLayersToDelete.Add(DataLayerInstance);
 				}
 			}
 		}
@@ -342,7 +342,7 @@ void FDataLayerMode::DeleteItems(const TArray<TWeakPtr<ISceneOutlinerTreeItem>>&
 	else if (!DataLayersToDelete.IsEmpty())
 	{
 		int32 PrevDeleteCount = SelectedDataLayersSet.Num();
-		for (UDataLayer* DataLayerToDelete : DataLayersToDelete)
+		for (UDataLayerInstance* DataLayerToDelete : DataLayersToDelete)
 		{
 			SelectedDataLayersSet.Remove(DataLayerToDelete);
 		}
@@ -404,7 +404,7 @@ FSceneOutlinerDragValidationInfo FDataLayerMode::ValidateDrop(const ISceneOutlin
 {
 	if (const FDataLayerTreeItem* DataLayerItem = DropTarget.CastTo<FDataLayerTreeItem>())
 	{
-		const UDataLayer* TargetDataLayer = DataLayerItem->GetDataLayer();
+		const UDataLayerInstance* TargetDataLayer = DataLayerItem->GetDataLayer();
 		if (!TargetDataLayer)
 		{
 			return FSceneOutlinerDragValidationInfo(ESceneOutlinerDropCompatibility::IncompatibleGeneric, FText());
@@ -427,7 +427,7 @@ FSceneOutlinerDragValidationInfo FDataLayerMode::ValidateDrop(const ISceneOutlin
 			}
 		}
 
-		const FText Text = bMoveOperation ? FText::Format(LOCTEXT("MoveToDataLayer", "Move to Data Layer \"{0}\""), FText::FromName(TargetDataLayer->GetDataLayerLabel())) : FText::Format(LOCTEXT("AssignToDataLayer", "Assign to Data Layer \"{0}\""), FText::FromName(TargetDataLayer->GetDataLayerLabel()));
+		const FText Text = bMoveOperation ? FText::Format(LOCTEXT("MoveToDataLayer", "Move to Data Layer \"{0}\""), FText::FromString(TargetDataLayer->GetDataLayerShortName())) : FText::Format(LOCTEXT("AssignToDataLayer", "Assign to Data Layer \"{0}\""), FText::FromString(TargetDataLayer->GetDataLayerShortName()));
 		return FSceneOutlinerDragValidationInfo(ESceneOutlinerDropCompatibility::Compatible, Text);
 	}
 
@@ -455,47 +455,54 @@ FSceneOutlinerDragValidationInfo FDataLayerMode::ValidateDrop(const ISceneOutlin
 	}
 	else // Moving a data layer(s)
 	{
-		TArray<UDataLayer*> PayloadDataLayers = GetDataLayersFromOperation(Payload.SourceOperation);
+		TArray<UDataLayerInstance*> PayloadDataLayers = GetDataLayersFromOperation(Payload.SourceOperation);
 		if (!PayloadDataLayers.IsEmpty())
 		{
 			const FDataLayerTreeItem* DataLayerItem = DropTarget.CastTo<FDataLayerTreeItem>();
 			const FDataLayerActorTreeItem* DataLayerActorTreeItem = DropTarget.CastTo<FDataLayerActorTreeItem>();
-			const UDataLayer* ParentDataLayer = DataLayerItem ? DataLayerItem->GetDataLayer() : nullptr;
+			const UDataLayerInstance* ParentDataLayer = DataLayerItem ? DataLayerItem->GetDataLayer() : nullptr;
 			if (!ParentDataLayer && DataLayerActorTreeItem)
 			{
 				ParentDataLayer = DataLayerActorTreeItem->GetDataLayer();
 			}
 
-			bool bCanSetParent = false;
-			for (UDataLayer* DataLayer : PayloadDataLayers)
+			for (UDataLayerInstance* DataLayerInstance : PayloadDataLayers)
 			{
-				if (DataLayer->IsLocked())
+				if (DataLayerInstance->IsLocked())
 				{
 					return FSceneOutlinerDragValidationInfo(ESceneOutlinerDropCompatibility::IncompatibleGeneric, LOCTEXT("CantMoveLockedDataLayer", "Can't move locked Data Layer"));
 				}
 
-				if (DataLayer->CanParent(ParentDataLayer))
+				if (ParentDataLayer != nullptr)
 				{
-					bCanSetParent = true;
+					if (!DataLayerInstance->CanParent(ParentDataLayer))
+					{
+						if (!DataLayerInstance->IsDataLayerTypeValidToParent(ParentDataLayer->GetType()))
+						{
+							return FSceneOutlinerDragValidationInfo(ESceneOutlinerDropCompatibility::IncompatibleGeneric, 
+								FText::Format(LOCTEXT("CantMoveToDataLayerDiffType", "Can't move a {0} Data Layer under a {1} Data Layer"), UEnum::GetDisplayValueAsText(DataLayerInstance->GetType()), UEnum::GetDisplayValueAsText(ParentDataLayer->GetType())));
+						}
+
+						if (ParentDataLayer == DataLayerInstance->GetParent() || ParentDataLayer == DataLayerInstance)
+						{
+							return FSceneOutlinerDragValidationInfo(ESceneOutlinerDropCompatibility::IncompatibleGeneric, LOCTEXT("CantMoveToSameDataLayer", "Can't move Data Layer to same Data Layer"));
+						}
+					}
+				
+					if (ParentDataLayer->IsLocked())
+					{
+						return FSceneOutlinerDragValidationInfo(ESceneOutlinerDropCompatibility::IncompatibleGeneric, LOCTEXT("CantMoveDataLayerToLockedDataLayer", "Can't move Data Layer to locked Data Layer"));
+					}
 				}
 			}
 
-			if (bCanSetParent)
+			if (ParentDataLayer != nullptr)
 			{
-				if (ParentDataLayer && ParentDataLayer->IsLocked())
-				{
-					return FSceneOutlinerDragValidationInfo(ESceneOutlinerDropCompatibility::IncompatibleGeneric, LOCTEXT("CantMoveDataLayerToLockedDataLayer", "Can't move Data Layer to locked Data Layer"));
-				}
-
-				if (ParentDataLayer)
-				{
-					return FSceneOutlinerDragValidationInfo(ESceneOutlinerDropCompatibility::Compatible, FText::Format(LOCTEXT("MoveDataLayerToDataLayer", "Move to Data Layer \"{0}\""), FText::FromName(ParentDataLayer->GetDataLayerLabel())));
-				}
-				return FSceneOutlinerDragValidationInfo(ESceneOutlinerDropCompatibility::Compatible, LOCTEXT("MoveDataLayerToRoot", "Move to root"));
+				return FSceneOutlinerDragValidationInfo(ESceneOutlinerDropCompatibility::Compatible, FText::Format(LOCTEXT("MoveDataLayerToDataLayer", "Move to Data Layer \"{0}\""), FText::FromString(ParentDataLayer->GetDataLayerShortName())));
 			}
 			else
 			{
-				return FSceneOutlinerDragValidationInfo(ESceneOutlinerDropCompatibility::IncompatibleGeneric, LOCTEXT("CantMoveToSameDataLayer", "Can't move Data Layer to same Data Layer"));
+				return FSceneOutlinerDragValidationInfo(ESceneOutlinerDropCompatibility::Compatible, LOCTEXT("MoveDataLayerToRoot", "Move to root"));
 			}
 		}
 	}
@@ -503,18 +510,18 @@ FSceneOutlinerDragValidationInfo FDataLayerMode::ValidateDrop(const ISceneOutlin
 	return FSceneOutlinerDragValidationInfo::Invalid();
 }
 
-TArray<UDataLayer*> FDataLayerMode::GetDataLayersFromOperation(const FDragDropOperation& Operation, bool bOnlyFindFirst /*= false*/) const
+TArray<UDataLayerInstance*> FDataLayerMode::GetDataLayersFromOperation(const FDragDropOperation& Operation, bool bOnlyFindFirst /*= false*/) const
 {
-	TArray<UDataLayer*> OutDataLayers;
+	TArray<UDataLayerInstance*> OutDataLayers;
 	
 	auto GetDataLayers = [this, &OutDataLayers](const FDataLayerDragDropOp& DataLayerOp)
 	{
-		const TArray<FName>& DataLayerLabels = DataLayerOp.DataLayerLabels;
-		for (FName DataLayerLabel : DataLayerLabels)
+		const TArray<FDataLayerDragDropOp::FDragDropInfo>& DragDropInfos = DataLayerOp.DataLayerDragDropInfos;
+		for (const FDataLayerDragDropOp::FDragDropInfo& DragDropInfo : DragDropInfos)
 		{
-			if (UDataLayer* DataLayer = DataLayerEditorSubsystem->GetDataLayerFromLabel(DataLayerLabel))
+			if (UDataLayerInstance* DataLayerInstance = DataLayerEditorSubsystem->GetDataLayerInstance(DragDropInfo.DataLayerInstanceName))
 			{
-				OutDataLayers.AddUnique(DataLayer);
+				OutDataLayers.AddUnique(DataLayerInstance);
 			}
 		}
 	};
@@ -625,14 +632,14 @@ TArray<AActor*> FDataLayerMode::GetActorsFromOperation(const FDragDropOperation&
 void FDataLayerMode::OnDrop(ISceneOutlinerTreeItem& DropTarget, const FSceneOutlinerDragDropPayload& Payload, const FSceneOutlinerDragValidationInfo& ValidationInfo) const
 {
 	const FDataLayerTreeItem* DataLayerItem = DropTarget.CastTo<FDataLayerTreeItem>();
-	UDataLayer* TargetDataLayer = DataLayerItem ? DataLayerItem->GetDataLayer() : nullptr;
+	UDataLayerInstance* TargetDataLayer = DataLayerItem ? DataLayerItem->GetDataLayer() : nullptr;
 
 	TArray<AActor*> ActorsToAdd = GetActorsFromOperation(Payload.SourceOperation);
 	if (!ActorsToAdd.IsEmpty()) // Adding actor(s) in data layer(s)
 	{
 		if (SceneOutliner->GetTree().IsItemSelected(const_cast<ISceneOutlinerTreeItem&>(DropTarget).AsShared()))
 		{
-			TArray<UDataLayer*> AllSelectedDataLayers = GetSelectedDataLayers(SceneOutliner);
+			TArray<UDataLayerInstance*> AllSelectedDataLayers = GetSelectedDataLayers(SceneOutliner);
 			if (AllSelectedDataLayers.Num() > 1)
 			{
 				const FScopedDataLayerTransaction Transaction(LOCTEXT("DataLayerOutlinerAddActorsToDataLayers", "Add Actors to Data Layers"), RepresentingWorld.Get());
@@ -667,8 +674,8 @@ void FDataLayerMode::OnDrop(ISceneOutlinerTreeItem& DropTarget, const FSceneOutl
 	}
 	else // Moving a data layer(s)
 	{
-		TArray<UDataLayer*> DataLayers = GetDataLayersFromOperation(Payload.SourceOperation);
-		SetParentDataLayer(DataLayers, TargetDataLayer);
+		TArray<UDataLayerInstance*> DataLayerInstances = GetDataLayersFromOperation(Payload.SourceOperation);
+		SetParentDataLayer(DataLayerInstances, TargetDataLayer);
 	}
 }
 
@@ -693,26 +700,26 @@ FReply FDataLayerMode::OnDragOverItem(const FDragDropEvent& Event, const ISceneO
 	return FReply::Handled();
 }
 
-void FDataLayerMode::SetParentDataLayer(const TArray<UDataLayer*> DataLayers, UDataLayer* ParentDataLayer) const
+void FDataLayerMode::SetParentDataLayer(const TArray<UDataLayerInstance*> DataLayerInstances, UDataLayerInstance* ParentDataLayer) const
 {
-	if (!DataLayers.IsEmpty())
+	if (!DataLayerInstances.IsEmpty())
 	{
-		TArray<UDataLayer*> ValidDataLayers;
-		ValidDataLayers.Reserve(DataLayers.Num());
-		for (UDataLayer* DataLayer : DataLayers)
+		TArray<UDataLayerInstance*> ValidDataLayers;
+		ValidDataLayers.Reserve(DataLayerInstances.Num());
+		for (UDataLayerInstance* DataLayerInstance : DataLayerInstances)
 		{
-			if (DataLayer->CanParent(ParentDataLayer))
+			if (DataLayerInstance->CanParent(ParentDataLayer))
 			{
-				ValidDataLayers.Add(DataLayer);
+				ValidDataLayers.Add(DataLayerInstance);
 			}
 		}
 
 		if (!ValidDataLayers.IsEmpty())
 		{
 			const FScopedDataLayerTransaction Transaction(LOCTEXT("DataLayerOutlinerChangeDataLayersParent", "Change Data Layers Parent"), RepresentingWorld.Get());
-			for (UDataLayer* DataLayer : ValidDataLayers)
+			for (UDataLayerInstance* DataLayerInstance : ValidDataLayers)
 			{
-				DataLayerEditorSubsystem->SetParentDataLayer(DataLayer, ParentDataLayer);
+				DataLayerEditorSubsystem->SetParentDataLayer(DataLayerInstance, ParentDataLayer);
 			}
 		}
 	}
@@ -759,7 +766,7 @@ struct FDataLayerActorPairSelector
 
 struct FWeakDataLayerSelector
 {
-	bool operator()(const TWeakPtr<ISceneOutlinerTreeItem>& Item, TWeakObjectPtr<UDataLayer>& DataOut) const
+	bool operator()(const TWeakPtr<ISceneOutlinerTreeItem>& Item, TWeakObjectPtr<UDataLayerInstance>& DataOut) const
 	{
 		if (TSharedPtr<ISceneOutlinerTreeItem> ItemPtr = Item.Pin())
 		{
@@ -782,8 +789,8 @@ TSharedPtr<FDragDropOperation> FDataLayerMode::CreateDragDropOperation(const FPo
 
 	if (DraggedObjects.Has<FDataLayerTreeItem>())
 	{
-		TArray<TWeakObjectPtr<UDataLayer>> DataLayers = DraggedObjects.GetData<TWeakObjectPtr<UDataLayer>>(FWeakDataLayerSelector());
-		if (DataLayers.FindByPredicate([&](const TWeakObjectPtr<UDataLayer>& DataLayer) { return DataLayer.IsValid() && DataLayer->IsLocked(); }))
+		TArray<TWeakObjectPtr<UDataLayerInstance>> DataLayerInstances = DraggedObjects.GetData<TWeakObjectPtr<UDataLayerInstance>>(FWeakDataLayerSelector());
+		if (DataLayerInstances.FindByPredicate([&](const TWeakObjectPtr<UDataLayerInstance>& DataLayer) { return DataLayer.IsValid() && DataLayer->IsLocked(); }))
 		{
 			return TSharedPtr<FDragDropOperation>();
 		}
@@ -792,12 +799,12 @@ TSharedPtr<FDragDropOperation> FDataLayerMode::CreateDragDropOperation(const FPo
 	auto GetDataLayerOperation = [&DraggedObjects]()
 	{
 		TSharedPtr<FDataLayerDragDropOp> DataLayerOperation = MakeShareable(new FDataLayerDragDropOp);
-		TArray<TWeakObjectPtr<UDataLayer>> DataLayers = DraggedObjects.GetData<TWeakObjectPtr<UDataLayer>>(FWeakDataLayerSelector());
-		for (const auto& DataLayer : DataLayers)
+		TArray<TWeakObjectPtr<UDataLayerInstance>> DataLayers = DraggedObjects.GetData<TWeakObjectPtr<UDataLayerInstance>>(FWeakDataLayerSelector());
+		for (const auto& DataLayerInstance : DataLayers)
 		{
-			if (DataLayer.IsValid())
+			if (DataLayerInstance.IsValid())
 			{
-				DataLayerOperation->DataLayerLabels.Add(DataLayer->GetDataLayerLabel());
+				DataLayerOperation->DataLayerDragDropInfos.Emplace(DataLayerInstance.Get());
 			}
 		}
 		DataLayerOperation->Construct();
@@ -857,12 +864,12 @@ TSharedPtr<FDragDropOperation> FDataLayerMode::CreateDragDropOperation(const FPo
 static const FName DefaultContextBaseMenuName("DataLayerOutliner.DefaultContextMenuBase");
 static const FName DefaultContextMenuName("DataLayerOutliner.DefaultContextMenu");
 
-TArray<UDataLayer*> FDataLayerMode::GetSelectedDataLayers(SSceneOutliner* InSceneOutliner) const
+TArray<UDataLayerInstance*> FDataLayerMode::GetSelectedDataLayers(SSceneOutliner* InSceneOutliner) const
 {
 	FSceneOutlinerItemSelection ItemSelection(InSceneOutliner->GetSelection());
 	TArray<FDataLayerTreeItem*> SelectedDataLayerItems;
 	ItemSelection.Get<FDataLayerTreeItem>(SelectedDataLayerItems);
-	TArray<UDataLayer*> ValidSelectedDataLayers;
+	TArray<UDataLayerInstance*> ValidSelectedDataLayers;
 	Algo::TransformIf(SelectedDataLayerItems, ValidSelectedDataLayers, [](const auto Item) { return Item && Item->GetDataLayer(); }, [](const auto Item) { return Item->GetDataLayer(); });
 	return MoveTemp(ValidSelectedDataLayers);
 }
@@ -897,16 +904,16 @@ void FDataLayerMode::RegisterContextMenu()
 			}
 
 			SSceneOutliner* SceneOutliner = Context->SceneOutliner.Pin().Get();
-			TArray<UDataLayer*> SelectedDataLayers = GetSelectedDataLayers(SceneOutliner);
-			const bool bSelectedDataLayersContainsLocked = !!SelectedDataLayers.FindByPredicate([](const UDataLayer* DataLayer) { return DataLayer->IsLocked(); });
+			TArray<UDataLayerInstance*> SelectedDataLayers = GetSelectedDataLayers(SceneOutliner);
+			const bool bSelectedDataLayersContainsLocked = !!SelectedDataLayers.FindByPredicate([](const UDataLayerInstance* DataLayer) { return DataLayer->IsLocked(); });
 
-			bool bHasActorEditorContextDataLayers = false;
-			TArray<const UDataLayer*> AllDataLayers;
+			bool bHasActorEditorContextDataLayers = false;			
+			TArray<const UDataLayerInstance*> AllDataLayers;
 			if (const AWorldDataLayers* WorldDataLayers = RepresentingWorld.IsValid() ? RepresentingWorld->GetWorldDataLayers() : nullptr)
 			{
-				WorldDataLayers->ForEachDataLayer([&AllDataLayers](UDataLayer* DataLayer)
+				WorldDataLayers->ForEachDataLayer([&AllDataLayers](UDataLayerInstance* DataLayerInstance)
 				{
-					AllDataLayers.Add(DataLayer);
+					AllDataLayers.Add(DataLayerInstance);
 					return true;
 				});
 
@@ -916,27 +923,26 @@ void FDataLayerMode::RegisterContextMenu()
 			{
 				FToolMenuSection& Section = InMenu->AddSection("DataLayers", LOCTEXT("DataLayers", "Data Layers"));
 				
-				auto CreateNewDataLayer = [this, SceneOutliner](UDataLayer* ParentDataLayer = nullptr)
+				auto CreateNewDataLayer = [this, SceneOutliner](UDataLayerInstance* ParentDataLayer = nullptr)
 				{
 					const FScopedDataLayerTransaction Transaction(LOCTEXT("CreateNewDataLayer", "Create New Data Layer"), RepresentingWorld.Get());
 					SelectedDataLayersSet.Empty();
 					SelectedDataLayerActors.Empty();
-					if (UDataLayer* NewDataLayer = DataLayerEditorSubsystem->CreateDataLayer())
+					if (UDataLayerInstance* NewDataLayerInstance = DataLayerEditorSubsystem->CreateDataLayer(ParentDataLayer))
 					{
-						SelectedDataLayersSet.Add(NewDataLayer);
-						DataLayerEditorSubsystem->SetParentDataLayer(NewDataLayer, ParentDataLayer);
+						SelectedDataLayersSet.Add(NewDataLayerInstance);
 						// Select it and open a rename when it gets refreshed
-						SceneOutliner->OnItemAdded(NewDataLayer, SceneOutliner::ENewItemAction::Select | SceneOutliner::ENewItemAction::Rename);
+						SceneOutliner->OnItemAdded(NewDataLayerInstance, SceneOutliner::ENewItemAction::Select | SceneOutliner::ENewItemAction::Rename);
 					}
 				};
 
 				Section.AddMenuEntry("CreateNewDataLayer", LOCTEXT("CreateNewDataLayer", "Create New Data Layer"), FText(), FSlateIcon(),
 					FUIAction(FExecuteAction::CreateLambda([this, CreateNewDataLayer]() { CreateNewDataLayer(); })));
 
-				UDataLayer* ParentDataLayer = (SelectedDataLayersSet.Num() == 1) ? const_cast<UDataLayer*>(SelectedDataLayersSet.CreateIterator()->Get()) : nullptr;
+				UDataLayerInstance* ParentDataLayer = (SelectedDataLayersSet.Num() == 1) ? const_cast<UDataLayerInstance*>(SelectedDataLayersSet.CreateIterator()->Get()) : nullptr;
 				if (ParentDataLayer)
 				{
-					Section.AddMenuEntry("CreateNewDataLayerUnderDataLayer", FText::Format(LOCTEXT("CreateNewDataLayerUnderDataLayer", "Create New Data Layer under \"{0}\""), FText::FromName(ParentDataLayer->GetDataLayerLabel())), FText(), FSlateIcon(),
+					Section.AddMenuEntry("CreateNewDataLayerUnderDataLayer", FText::Format(LOCTEXT("CreateNewDataLayerUnderDataLayer", "Create New Data Layer under \"{0}\""), FText::FromString(ParentDataLayer->GetDataLayerShortName())), FText(), FSlateIcon(),
 					FUIAction(FExecuteAction::CreateLambda([this, CreateNewDataLayer, ParentDataLayer]() { CreateNewDataLayer(ParentDataLayer); })));
 				}
 
@@ -944,7 +950,7 @@ void FDataLayerMode::RegisterContextMenu()
 					FUIAction(
 						FExecuteAction::CreateLambda([this]() {
 							const FScopedDataLayerTransaction Transaction(LOCTEXT("AddSelectedActorsToNewDataLayer", "Add Selected Actors to New Data Layer"), RepresentingWorld.Get());
-							if (UDataLayer* NewDataLayer = DataLayerEditorSubsystem->CreateDataLayer())
+							if (UDataLayerInstance* NewDataLayer = DataLayerEditorSubsystem->CreateDataLayer())
 							{
 								DataLayerEditorSubsystem->AddSelectedActorsToDataLayer(NewDataLayer);
 							}}),
@@ -967,7 +973,7 @@ void FDataLayerMode::RegisterContextMenu()
 					Section.AddSubMenu("AddSelectedActorsTo", LOCTEXT("AddSelectedActorsTo", "Add Selected Actors To"), FText(),
 						FNewToolMenuDelegate::CreateLambda([this](UToolMenu* InSubMenu)
 						{
-							CreateDataLayerPicker(InSubMenu, FOnDataLayerPicked::CreateLambda([this](UDataLayer* TargetDataLayer)
+							CreateDataLayerPicker(InSubMenu, FOnDataLayerPicked::CreateLambda([this](UDataLayerInstance* TargetDataLayer)
 							{
 								check(TargetDataLayer);
 								TArray<AActor*> Actors;
@@ -985,14 +991,14 @@ void FDataLayerMode::RegisterContextMenu()
 					Section.AddSubMenu("MoveSelectedDataLayersTo", LOCTEXT("MoveSelectedDataLayersTo", "Move Data Layers To"), FText(),
 						FNewToolMenuDelegate::CreateLambda([this](UToolMenu* InSubMenu)
 						{
-							CreateDataLayerPicker(InSubMenu, FOnDataLayerPicked::CreateLambda([this](UDataLayer* TargetDataLayer)
+							CreateDataLayerPicker(InSubMenu, FOnDataLayerPicked::CreateLambda([this](UDataLayerInstance* TargetDataLayer)
 							{
-								TArray<UDataLayer*> DataLayers;
-								for (TWeakObjectPtr<const UDataLayer>& DataLayer : SelectedDataLayersSet)
+								TArray<UDataLayerInstance*> DataLayers;
+								for (TWeakObjectPtr<const UDataLayerInstance>& DataLayer : SelectedDataLayersSet)
 								{
 									if (DataLayer.IsValid() && !DataLayer->IsLocked() && DataLayer != TargetDataLayer)
 									{
-										DataLayers.Add(const_cast<UDataLayer*>(DataLayer.Get()));
+										DataLayers.Add(const_cast<UDataLayerInstance*>(DataLayer.Get()));
 									}
 								}
 								SetParentDataLayer(DataLayers, TargetDataLayer);
@@ -1114,15 +1120,15 @@ void FDataLayerMode::RegisterContextMenu()
 						FExecuteAction::CreateLambda([this, SceneOutliner, SelectedDataLayers]() 
 						{
 							const FScopedDataLayerTransaction Transaction(LOCTEXT("MakeCurrentDataLayers", "Make Current Data Layer(s)"), RepresentingWorld.Get());
-							for (TWeakObjectPtr<const UDataLayer>& DataLayer : SelectedDataLayersSet)
+							for (TWeakObjectPtr<const UDataLayerInstance>& DataLayer : SelectedDataLayersSet)
 							{
 								if (DataLayer.IsValid() && !DataLayer->IsLocked())
 								{
-									DataLayerEditorSubsystem->AddToActorEditorContext(const_cast<UDataLayer*>(DataLayer.Get()));
+									DataLayerEditorSubsystem->AddToActorEditorContext(const_cast<UDataLayerInstance*>(DataLayer.Get()));
 								}
 							}
 						}),
-						FCanExecuteAction::CreateLambda([this] { return Algo::AnyOf(SelectedDataLayersSet, [](const TWeakObjectPtr<const UDataLayer>& DataLayer) { return DataLayer.IsValid() && !DataLayer->IsLocked() && !DataLayer->IsInActorEditorContext(); }); })
+						FCanExecuteAction::CreateLambda([this] { return Algo::AnyOf(SelectedDataLayersSet, [](const TWeakObjectPtr<const UDataLayerInstance>& DataLayer) { return DataLayer.IsValid() && !DataLayer->IsLocked() && !DataLayer->IsInActorEditorContext(); }); })
 					));
 
 				Section.AddMenuEntry("RemoveCurrentDataLayers", LOCTEXT("RemoveCurrentDataLayers", "Remove Current Data Layer(s)"), FText(), FSlateIcon(),
@@ -1130,15 +1136,15 @@ void FDataLayerMode::RegisterContextMenu()
 						FExecuteAction::CreateLambda([this, SceneOutliner, SelectedDataLayers]()
 						{
 							const FScopedDataLayerTransaction Transaction(LOCTEXT("RemoveCurrentDataLayers", "Remove Current Data Layer(s)"), RepresentingWorld.Get());
-							for (TWeakObjectPtr<const UDataLayer>& DataLayer : SelectedDataLayersSet)
+							for (TWeakObjectPtr<const UDataLayerInstance>& DataLayer : SelectedDataLayersSet)
 							{
 								if (DataLayer.IsValid() && !DataLayer->IsLocked())
 								{
-									DataLayerEditorSubsystem->RemoveFromActorEditorContext(const_cast<UDataLayer*>(DataLayer.Get()));
+									DataLayerEditorSubsystem->RemoveFromActorEditorContext(const_cast<UDataLayerInstance*>(DataLayer.Get()));
 								}
 							}
 						}),
-						FCanExecuteAction::CreateLambda([this] { return Algo::AnyOf(SelectedDataLayersSet, [](const TWeakObjectPtr<const UDataLayer>& DataLayer) { return DataLayer.IsValid() && !DataLayer->IsLocked() && DataLayer->IsInActorEditorContext(); }); })
+						FCanExecuteAction::CreateLambda([this] { return Algo::AnyOf(SelectedDataLayersSet, [](const TWeakObjectPtr<const UDataLayerInstance>& DataLayer) { return DataLayer.IsValid() && !DataLayer->IsLocked() && DataLayer->IsInActorEditorContext(); }); })
 					));
 
 				Section.AddMenuEntry("ClearCurrentDataLayers", LOCTEXT("ClearCurrentDataLayers", "Clear Current Data Layers"), FText(), FSlateIcon(),
@@ -1148,9 +1154,9 @@ void FDataLayerMode::RegisterContextMenu()
 							check(!AllDataLayers.IsEmpty());
 							{
 								const FScopedDataLayerTransaction Transaction(LOCTEXT("ClearCurrentDataLayers", "Clear Current Data Layers"), RepresentingWorld.Get());
-								for (const UDataLayer* DataLayer : AllDataLayers)
+								for (const UDataLayerInstance* DataLayer : AllDataLayers)
 								{
-									DataLayerEditorSubsystem->RemoveFromActorEditorContext(const_cast<UDataLayer*>(DataLayer));
+									DataLayerEditorSubsystem->RemoveFromActorEditorContext(const_cast<UDataLayerInstance*>(DataLayer));
 								}
 							}}),
 						FCanExecuteAction::CreateLambda([bHasActorEditorContextDataLayers] { return bHasActorEditorContextDataLayers; })
@@ -1263,10 +1269,10 @@ void FDataLayerMode::CreateViewContent(FMenuBuilder& MenuBuilder)
 		EUserInterfaceActionType::ToggleButton
 	);
 
-	TArray<UDataLayer*> AllDataLayers;
+	TArray<UDataLayerInstance*> AllDataLayers;
 	if (const AWorldDataLayers* WorldDataLayers = RepresentingWorld.IsValid() ? RepresentingWorld->GetWorldDataLayers() : nullptr)
 	{
-		WorldDataLayers->ForEachDataLayer([&AllDataLayers](UDataLayer* DataLayer)
+		WorldDataLayers->ForEachDataLayer([&AllDataLayers](UDataLayerInstance* DataLayer)
 		{
 			AllDataLayers.Add(DataLayer);
 			return true;
@@ -1470,7 +1476,7 @@ void FDataLayerMode::ChooseRepresentingWorld()
 	}
 }
 
-bool FDataLayerMode::ShouldExpandDataLayer(const UDataLayer* DataLayer) const
+bool FDataLayerMode::ShouldExpandDataLayer(const UDataLayerInstance* DataLayer) const
 {
 	if (bHighlightSelectedDataLayers || bShowOnlySelectedActors)
 	{
@@ -1486,12 +1492,12 @@ bool FDataLayerMode::ShouldExpandDataLayer(const UDataLayer* DataLayer) const
 	return false;
 }
 
-bool FDataLayerMode::ContainsSelectedChildDataLayer(const UDataLayer* DataLayer) const
+bool FDataLayerMode::ContainsSelectedChildDataLayer(const UDataLayerInstance* DataLayer) const
 {
 	if (DataLayer)
 	{
 		bool bFoundSelected = false;
-		DataLayer->ForEachChild([this, &bFoundSelected](const UDataLayer* Child)
+		DataLayer->ForEachChild([this, &bFoundSelected](const UDataLayerInstance* Child)
 		{
 			if (DataLayerEditorSubsystem->DoesDataLayerContainSelectedActors(Child) || ContainsSelectedChildDataLayer(Child))
 			{
@@ -1507,7 +1513,7 @@ bool FDataLayerMode::ContainsSelectedChildDataLayer(const UDataLayer* DataLayer)
 
 TSharedRef<FSceneOutlinerFilter> FDataLayerMode::CreateShowOnlySelectedActorsFilter()
 {
-	auto IsActorSelected = [](const AActor* InActor, const UDataLayer* InDataLayer)
+	auto IsActorSelected = [](const AActor* InActor, const UDataLayerInstance* InDataLayer)
 	{
 		return InActor && InActor->IsSelected();
 	};
@@ -1522,14 +1528,14 @@ void FDataLayerMode::SynchronizeSelection()
 	}
 
 	TArray<AActor*> Actors;
-	TSet<const UDataLayer*> ActorDataLayersIncludingParents;
+	TSet<const UDataLayerInstance*> ActorDataLayersIncludingParents;
 	GEditor->GetSelectedActors()->GetSelectedObjects<AActor>(Actors);
 	for (const AActor* Actor : Actors)
 	{
-		TArray<const UDataLayer*> ActorDataLayers = Actor->GetDataLayerObjects();
-		for (const UDataLayer* DataLayer : ActorDataLayers)
+		TArray<const UDataLayerInstance*> ActorDataLayers = Actor->GetDataLayerInstances();
+		for (const UDataLayerInstance* DataLayer : ActorDataLayers)
 		{
-			const UDataLayer* CurrentDataLayer = DataLayer;
+			const UDataLayerInstance* CurrentDataLayer = DataLayer;
 			while (CurrentDataLayer)
 			{
 				bool bIsAlreadyInSet = false;
@@ -1590,7 +1596,7 @@ TSharedRef<SWidget> FDataLayerPickingMode::CreateDataLayerPickerWidget(FOnDataLa
 		return new FDataLayerPickingMode(FDataLayerModeParams(Outliner, nullptr, nullptr, FOnSceneOutlinerItemPicked::CreateLambda([OnDataLayerPicked](const FSceneOutlinerTreeItemRef& NewParent)
 		{
 			FDataLayerTreeItem* DataLayerItem = NewParent->CastTo<FDataLayerTreeItem>();
-			if (UDataLayer* DataLayer = DataLayerItem ? DataLayerItem->GetDataLayer() : nullptr)
+			if (UDataLayerInstance* DataLayer = DataLayerItem ? DataLayerItem->GetDataLayer() : nullptr)
 			{
 				OnDataLayerPicked.ExecuteIfBound(DataLayer);
 			}
@@ -1624,7 +1630,7 @@ void FDataLayerPickingMode::OnItemSelectionChanged(FSceneOutlinerTreeItemPtr Tre
 			{
 				if (const FDataLayerTreeItem* DataLayerItem = FirstItem->CastTo<FDataLayerTreeItem>())
 				{
-					UDataLayer* DataLayer = DataLayerItem->GetDataLayer();
+					UDataLayerInstance* DataLayer = DataLayerItem->GetDataLayer();
 					if (DataLayer && !DataLayer->IsLocked())
 					{
 						OnItemPicked.ExecuteIfBound(FirstItem.ToSharedRef());
diff --git a/Engine/Source/Editor/DataLayerEditor/Private/DataLayer/DataLayerMode.h b/Engine/Source/Editor/DataLayerEditor/Private/DataLayer/DataLayerMode.h
--- a/Engine/Source/Editor/DataLayerEditor/Private/DataLayer/DataLayerMode.h
+++ b/Engine/Source/Editor/DataLayerEditor/Private/DataLayer/DataLayerMode.h
@@ -24,7 +24,7 @@ struct FDataLayerModeParams
 	FOnSceneOutlinerItemPicked OnItemPicked;
 };
 
-DECLARE_DELEGATE_OneParam(FOnDataLayerPicked, UDataLayer*);
+DECLARE_DELEGATE_OneParam(FOnDataLayerPicked, UDataLayerInstance*);
 
 class FDataLayerMode : public ISceneOutlinerMode
 {
@@ -101,13 +101,13 @@ private:
 	bool IsWorldChecked(TWeakObjectPtr<UWorld> World) const;
 	TArray<FDataLayerActorMoveElement> GetDataLayerActorPairsFromOperation(const FDragDropOperation& Operation) const;
 	TArray<AActor*> GetActorsFromOperation(const FDragDropOperation& Operation, bool bOnlyFindFirst = false) const;
-	TArray<UDataLayer*> GetDataLayersFromOperation(const FDragDropOperation& Operation, bool bOnlyFindFirst = false) const;
-	TArray<UDataLayer*> GetSelectedDataLayers(SSceneOutliner* InSceneOutliner) const;
-	void SetParentDataLayer(const TArray<UDataLayer*> DataLayers, UDataLayer* ParentDataLayer) const;
+	TArray<UDataLayerInstance*> GetDataLayersFromOperation(const FDragDropOperation& Operation, bool bOnlyFindFirst = false) const;
+	TArray<UDataLayerInstance*> GetSelectedDataLayers(SSceneOutliner* InSceneOutliner) const;
+	void SetParentDataLayer(const TArray<UDataLayerInstance*> DataLayers, UDataLayerInstance* ParentDataLayer) const;
 	void OnLevelSelectionChanged(UObject* Obj);
 	static void CreateDataLayerPicker(UToolMenu* InMenu, FOnDataLayerPicked OnDataLayerPicked, bool bInShowRoot = false);
-	bool ShouldExpandDataLayer(const UDataLayer* DataLayer) const;
-	bool ContainsSelectedChildDataLayer(const UDataLayer* DataLayer) const;
+	bool ShouldExpandDataLayer(const UDataLayerInstance* DataLayer) const;
+	bool ContainsSelectedChildDataLayer(const UDataLayerInstance* DataLayer) const;
 	void RefreshSelection();
 	FSceneOutlinerDragValidationInfo ValidateDrop(const ISceneOutlinerTreeItem& DropTarget, bool bMoveOperation = false) const;
 
@@ -131,10 +131,10 @@ private:
 	/** Number of datalayers which have passed through the filters */
 	uint32 FilteredDataLayerCount = 0;
 	/** List of datalayers which passed the regular filters and may or may not have passed the search filter */
-	TSet<TWeakObjectPtr<UDataLayer>> ApplicableDataLayers;
+	TSet<TWeakObjectPtr<UDataLayerInstance>> ApplicableDataLayers;
 
-	TSet<TWeakObjectPtr<const UDataLayer>> SelectedDataLayersSet;
-	typedef TPair<TWeakObjectPtr<const UDataLayer>, TWeakObjectPtr<const AActor>> FSelectedDataLayerActor;
+	TSet<TWeakObjectPtr<const UDataLayerInstance>> SelectedDataLayersSet;
+	typedef TPair<TWeakObjectPtr<const UDataLayerInstance>, TWeakObjectPtr<const AActor>> FSelectedDataLayerActor;
 	TSet<FSelectedDataLayerActor> SelectedDataLayerActors;
 };
 
diff --git a/Engine/Source/Editor/DataLayerEditor/Private/DataLayer/DataLayerNameEditSink.h b/Engine/Source/Editor/DataLayerEditor/Private/DataLayer/DataLayerNameEditSink.h
--- a/Engine/Source/Editor/DataLayerEditor/Private/DataLayer/DataLayerNameEditSink.h
+++ b/Engine/Source/Editor/DataLayerEditor/Private/DataLayer/DataLayerNameEditSink.h
@@ -3,7 +3,6 @@
 #pragma once
 
 #include "DataLayer/DataLayerEditorSubsystem.h"
-#include "WorldPartition/DataLayer/DataLayer.h"
 #include "DataLayerTransaction.h"
 #include "IObjectNameEditSink.h"
 
@@ -13,38 +12,47 @@ class FDataLayerNameEditSink : public UE::EditorWidgets::IObjectNameEditSink
 {
 	virtual UClass* GetSupportedClass() const override
 	{
-		return UDataLayer::StaticClass();
+		return UDataLayerInstance::StaticClass();
 	}
 
 	virtual FText GetObjectDisplayName(UObject* Object) const override
 	{
-		return FText::FromName(CastChecked<UDataLayer>(Object)->GetDataLayerLabel());
+		UDataLayerInstance* DataLayerInstance = CastChecked<UDataLayerInstance>(Object);
+		if (DataLayerInstance->SupportRelabeling())
+		{
+			return FText::Format(FText::FromString("{0}"), FText::FromString(DataLayerInstance->GetDataLayerShortName()));
+		}
+		
+		return FText::Format(FText::FromString("{0} ({1})"), FText::FromString(DataLayerInstance->GetDataLayerShortName()), FText::FromString(DataLayerInstance->GetDataLayerFullName()));
 	}
 
 	virtual bool IsObjectDisplayNameReadOnly(UObject* Object) const override
 	{
-		UDataLayer* DataLayer = CastChecked<UDataLayer>(Object);
-		return DataLayer->IsLocked();
+		UDataLayerInstance* DataLayerInstance = CastChecked<UDataLayerInstance>(Object);
+		return !DataLayerInstance->SupportRelabeling() || DataLayerInstance->IsLocked();
 	};
 
 	virtual bool SetObjectDisplayName(UObject* Object, FString DisplayName) override
 	{
-		UDataLayer* DataLayer = CastChecked<UDataLayer>(Object);
-
-		if (DataLayer->GetDataLayerLabel().ToString() == DisplayName)
+		UDataLayerInstance* DataLayerInstance = CastChecked<UDataLayerInstance>(Object);
+		if(DataLayerInstance->SupportRelabeling())
 		{
-			return false;
+			if (!DisplayName.Equals(DataLayerInstance->GetDataLayerShortName(), ESearchCase::CaseSensitive))
+			{
+				const FScopedDataLayerTransaction Transaction(LOCTEXT("DataLayerNameEditSinkRenameDataLayerTransaction", "Rename Data Layer"), DataLayerInstance->GetWorld());
+
+				PRAGMA_DISABLE_DEPRECATION_WARNINGS
+				return UDataLayerEditorSubsystem::Get()->RenameDataLayer(DataLayerInstance, *DisplayName);
+				PRAGMA_ENABLE_DEPRECATION_WARNINGS
+			}
 		}
-
-		const FScopedDataLayerTransaction Transaction(LOCTEXT("DataLayerNameEditSinkRenameDataLayerTransaction", "Rename Data Layer"), DataLayer->GetWorld());
-		UDataLayerEditorSubsystem::Get()->RenameDataLayer(DataLayer, *DisplayName);
-
-		return true;
+		
+		return false;
 	}
 
 	virtual FText GetObjectNameTooltip(UObject* Object) const override
 	{
-		return FText::Format(LOCTEXT("EditableDataLayerLabel_TooltipFmt", "Rename the selected {0}"), FText::FromString(Object->GetClass()->GetName()));
+		return IsObjectDisplayNameReadOnly(Object) ? LOCTEXT("NonEditableDataLayerLabel_TooltipFmt", "Data Layer Name") : FText::Format(LOCTEXT("EditableDataLayerLabel_TooltipFmt", "Rename the selected {0}"), FText::FromString(Object->GetClass()->GetName()));
 	}
 };
 
diff --git a/Engine/Source/Editor/DataLayerEditor/Private/DataLayer/DataLayerOutlinerDeleteButtonColumn.cpp b/Engine/Source/Editor/DataLayerEditor/Private/DataLayer/DataLayerOutlinerDeleteButtonColumn.cpp
--- a/Engine/Source/Editor/DataLayerEditor/Private/DataLayer/DataLayerOutlinerDeleteButtonColumn.cpp
+++ b/Engine/Source/Editor/DataLayerEditor/Private/DataLayer/DataLayerOutlinerDeleteButtonColumn.cpp
@@ -43,13 +43,13 @@ const TSharedRef<SWidget> FDataLayerOutlinerDeleteButtonColumn::ConstructRowWidg
 			.Visibility_Lambda([this, TreeItem, DataLayerActorItem]()
 			{
 				AActor* Actor = DataLayerActorItem->GetActor();
-				const UDataLayer* DataLayer = DataLayerActorItem->GetDataLayer();
+				const UDataLayerInstance* DataLayer = DataLayerActorItem->GetDataLayer();
 				return (Actor && DataLayer && !DataLayer->IsLocked()) ? EVisibility::Visible : EVisibility::Collapsed;
 			})
 			.OnClicked_Lambda([this, TreeItem, DataLayerActorItem]()
 			{
 				AActor* Actor = DataLayerActorItem->GetActor();
-				const UDataLayer* DataLayer = DataLayerActorItem->GetDataLayer();
+				const UDataLayerInstance* DataLayer = DataLayerActorItem->GetDataLayer();
 				if (Actor && DataLayer)
 				{
 					UDataLayerEditorSubsystem* DataLayerEditorSubsystem = UDataLayerEditorSubsystem::Get();
diff --git a/Engine/Source/Editor/DataLayerEditor/Private/DataLayer/DataLayerOutlinerIsLoadedInEditorColumn.cpp b/Engine/Source/Editor/DataLayerEditor/Private/DataLayer/DataLayerOutlinerIsLoadedInEditorColumn.cpp
--- a/Engine/Source/Editor/DataLayerEditor/Private/DataLayer/DataLayerOutlinerIsLoadedInEditorColumn.cpp
+++ b/Engine/Source/Editor/DataLayerEditor/Private/DataLayer/DataLayerOutlinerIsLoadedInEditorColumn.cpp
@@ -49,22 +49,22 @@ const TSharedRef<SWidget> FDataLayerOutlinerIsLoadedInEditorColumn::ConstructRow
 				.IsEnabled_Lambda([this, TreeItem]()
 				{
 					FDataLayerTreeItem* DataLayerTreeItem = TreeItem->CastTo<FDataLayerTreeItem>();
-					const UDataLayer* DataLayer = DataLayerTreeItem->GetDataLayer();
-					const UDataLayer* ParentDataLayer = DataLayer ? DataLayer->GetParent() : nullptr;
+					const UDataLayerInstance* DataLayer = DataLayerTreeItem->GetDataLayer();
+					const UDataLayerInstance* ParentDataLayer = DataLayer ? DataLayer->GetParent() : nullptr;
 					const bool bIsParentLoaded = ParentDataLayer ? ParentDataLayer->IsEffectiveLoadedInEditor() : true;
 					return bIsParentLoaded && DataLayer && DataLayer->GetWorld() && !DataLayer->GetWorld()->IsPlayInEditor();
 				})
 				.IsChecked_Lambda([this, TreeItem]()
 				{
 					FDataLayerTreeItem* DataLayerTreeItem = TreeItem->CastTo<FDataLayerTreeItem>();
-					UDataLayer* DataLayer = DataLayerTreeItem->GetDataLayer();
+					UDataLayerInstance* DataLayer = DataLayerTreeItem->GetDataLayer();
 					return DataLayer && DataLayer->IsEffectiveLoadedInEditor() ? ECheckBoxState::Checked : ECheckBoxState::Unchecked;
 				})
 				.OnCheckStateChanged_Lambda([this, TreeItem](ECheckBoxState NewState)
 				{
 					bool bNewState = (NewState == ECheckBoxState::Checked);
 					FDataLayerTreeItem* DataLayerTreeItem = TreeItem->CastTo<FDataLayerTreeItem>();
-					if (UDataLayer* DataLayer = DataLayerTreeItem->GetDataLayer())
+					if (UDataLayerInstance* DataLayer = DataLayerTreeItem->GetDataLayer())
 					{
 						UWorld* World = DataLayer->GetWorld();
 						UDataLayerEditorSubsystem* DataLayerEditorSubsystem = UDataLayerEditorSubsystem::Get();
@@ -74,11 +74,11 @@ const TSharedRef<SWidget> FDataLayerOutlinerIsLoadedInEditorColumn::ConstructRow
 							// Toggle IsLoadedInEditor flag of selected DataLayers to the same state as the given DataLayer
 							const bool bIsLoadedInEditor = DataLayer->IsLoadedInEditor();
 
-							TArray<UDataLayer*> AllSelectedDataLayers;
+							TArray<UDataLayerInstance*> AllSelectedDataLayers;
 							for (auto& SelectedItem : Tree.GetSelectedItems())
 							{
 								FDataLayerTreeItem* SelectedDataLayerTreeItem = SelectedItem->CastTo<FDataLayerTreeItem>();
-								UDataLayer* SelectedDataLayer = SelectedDataLayerTreeItem ? SelectedDataLayerTreeItem->GetDataLayer() : nullptr;
+								UDataLayerInstance* SelectedDataLayer = SelectedDataLayerTreeItem ? SelectedDataLayerTreeItem->GetDataLayer() : nullptr;
 								if (SelectedDataLayer && SelectedDataLayer->IsLoadedInEditor() == bIsLoadedInEditor)
 								{
 									AllSelectedDataLayers.Add(SelectedDataLayer);
diff --git a/Engine/Source/Editor/DataLayerEditor/Private/DataLayer/DataLayerOutlinerIsVisibleColumn.cpp b/Engine/Source/Editor/DataLayerEditor/Private/DataLayer/DataLayerOutlinerIsVisibleColumn.cpp
--- a/Engine/Source/Editor/DataLayerEditor/Private/DataLayer/DataLayerOutlinerIsVisibleColumn.cpp
+++ b/Engine/Source/Editor/DataLayerEditor/Private/DataLayer/DataLayerOutlinerIsVisibleColumn.cpp
@@ -29,8 +29,8 @@ protected:
 		auto TreeItem = WeakTreeItem.Pin();
 		if (FDataLayerTreeItem* DataLayerTreeItem = TreeItem.IsValid() ? TreeItem->CastTo<FDataLayerTreeItem>() : nullptr)
 		{
-			const UDataLayer* DataLayer = DataLayerTreeItem->GetDataLayer();
-			const UDataLayer* ParentDataLayer = DataLayer ? DataLayer->GetParent() : nullptr;
+			const UDataLayerInstance* DataLayer = DataLayerTreeItem->GetDataLayer();
+			const UDataLayerInstance* ParentDataLayer = DataLayer ? DataLayer->GetParent() : nullptr;
 			const bool bIsParentVisible = ParentDataLayer ? ParentDataLayer->IsEffectiveVisible() : true;
 			return bIsParentVisible && DataLayer && DataLayer->GetWorld() && !DataLayer->GetWorld()->IsPlayInEditor() && DataLayer->IsEffectiveLoadedInEditor();
 		}
@@ -43,7 +43,7 @@ protected:
 		auto TreeItem = WeakTreeItem.Pin();
 		if (FDataLayerTreeItem* DataLayerTreeItem = TreeItem.IsValid() ? TreeItem->CastTo<FDataLayerTreeItem>() : nullptr)
 		{
-			UDataLayer* DataLayer = DataLayerTreeItem->GetDataLayer();
+			UDataLayerInstance* DataLayer = DataLayerTreeItem->GetDataLayer();
 			bIsEffectiveVisible = DataLayer && DataLayer->IsEffectiveVisible();
 		}
 
diff --git a/Engine/Source/Editor/DataLayerEditor/Private/DataLayer/DataLayerPropertyTypeCustomization.cpp b/Engine/Source/Editor/DataLayerEditor/Private/DataLayer/DataLayerPropertyTypeCustomization.cpp
--- a/Engine/Source/Editor/DataLayerEditor/Private/DataLayer/DataLayerPropertyTypeCustomization.cpp
+++ b/Engine/Source/Editor/DataLayerEditor/Private/DataLayer/DataLayerPropertyTypeCustomization.cpp
@@ -5,7 +5,7 @@
 #include "DataLayer/DataLayerDragDropOp.h"
 #include "DataLayer/DataLayerEditorSubsystem.h"
 #include "DataLayerEditorModule.h"
-#include "WorldPartition/DataLayer/DataLayer.h"
+#include "WorldPartition/DataLayer/DataLayerInstance.h"
 #include "DragAndDrop/CompositeDragDropOp.h"
 #include "Algo/Accumulate.h"
 #include "Modules/ModuleManager.h"
@@ -64,7 +64,7 @@ void FDataLayerPropertyTypeCustomization::CustomizeHeader(TSharedRef<IPropertyHa
 				.IsEnabled_Lambda([this]
 				{
 					FPropertyAccess::Result PropertyAccessResult;
-					const UDataLayer* DataLayer = GetDataLayerFromPropertyHandle(&PropertyAccessResult);
+					const UDataLayerInstance* DataLayer = GetDataLayerFromPropertyHandle(&PropertyAccessResult);
 					return (!DataLayer || !DataLayer->IsLocked());
 				})
 				.ToolTipText(LOCTEXT("ComboButtonTip", "Drag and drop a Data Layer onto this property, or choose one from the drop down."))
@@ -88,7 +88,7 @@ void FDataLayerPropertyTypeCustomization::CustomizeHeader(TSharedRef<IPropertyHa
 				.Visibility_Lambda([this] 
 				{
 					FPropertyAccess::Result PropertyAccessResult;
-					const UDataLayer* DataLayer = GetDataLayerFromPropertyHandle(&PropertyAccessResult);
+					const UDataLayerInstance* DataLayer = GetDataLayerFromPropertyHandle(&PropertyAccessResult);
 					return (DataLayer && DataLayer->IsLocked()) ? EVisibility::Visible : EVisibility::Collapsed;
 				})
 				.ColorAndOpacity(this, &FDataLayerPropertyTypeCustomization::GetForegroundColor)
@@ -124,7 +124,7 @@ void FDataLayerPropertyTypeCustomization::CustomizeHeader(TSharedRef<IPropertyHa
 void FDataLayerPropertyTypeCustomization::OnBrowse()
 {
 	FPropertyAccess::Result PropertyAccessResult;
-	if (const UDataLayer* DataLayer = GetDataLayerFromPropertyHandle(&PropertyAccessResult))
+	if (const UDataLayerInstance* DataLayer = GetDataLayerFromPropertyHandle(&PropertyAccessResult))
 	{
 		FLevelEditorModule& LevelEditorModule = FModuleManager::GetModuleChecked<FLevelEditorModule>("LevelEditor");
 		LevelEditorModule.GetLevelEditorTabManager()->TryInvokeTab(FTabId("LevelEditorDataLayerBrowser"));
@@ -134,7 +134,7 @@ void FDataLayerPropertyTypeCustomization::OnBrowse()
 	}
 }
 
-UDataLayer* FDataLayerPropertyTypeCustomization::GetDataLayerFromPropertyHandle(FPropertyAccess::Result* OutPropertyAccessResult) const
+UDataLayerInstance* FDataLayerPropertyTypeCustomization::GetDataLayerFromPropertyHandle(FPropertyAccess::Result* OutPropertyAccessResult) const
 {
 	FName DataLayerName;
 	FPropertyAccess::Result Result = PropertyHandle->GetValue(DataLayerName);
@@ -144,8 +144,8 @@ UDataLayer* FDataLayerPropertyTypeCustomization::GetDataLayerFromPropertyHandle(
 	}
 	if (Result == FPropertyAccess::Success)
 	{
-		UDataLayer* DataLayer = UDataLayerEditorSubsystem::Get()->GetDataLayerFromName(DataLayerName);
-		return DataLayer;
+		UDataLayerInstance* DataLayerInstance = UDataLayerEditorSubsystem::Get()->GetDataLayerInstance(DataLayerName);
+		return DataLayerInstance;
 	}
 	return nullptr;
 }
@@ -153,7 +153,7 @@ UDataLayer* FDataLayerPropertyTypeCustomization::GetDataLayerFromPropertyHandle(
 const FSlateBrush* FDataLayerPropertyTypeCustomization::GetDataLayerIcon() const
 {
 	FPropertyAccess::Result PropertyAccessResult;
-	const UDataLayer* DataLayer = GetDataLayerFromPropertyHandle(&PropertyAccessResult);
+	const UDataLayerInstance* DataLayer = GetDataLayerFromPropertyHandle(&PropertyAccessResult);
 	if (!DataLayer)
 	{
 		return FEditorStyle::GetBrush(TEXT("DataLayer.Editor"));
@@ -168,18 +168,18 @@ const FSlateBrush* FDataLayerPropertyTypeCustomization::GetDataLayerIcon() const
 FText FDataLayerPropertyTypeCustomization::GetDataLayerText() const
 {
 	FPropertyAccess::Result PropertyAccessResult;
-	const UDataLayer* DataLayer = GetDataLayerFromPropertyHandle(&PropertyAccessResult);
+	const UDataLayerInstance* DataLayer = GetDataLayerFromPropertyHandle(&PropertyAccessResult);
 	if (PropertyAccessResult == FPropertyAccess::MultipleValues)
 	{
 		return NSLOCTEXT("PropertyEditor", "MultipleValues", "Multiple Values");
 	}
-	return UDataLayer::GetDataLayerText(DataLayer);
+	return UDataLayerInstance::GetDataLayerText(DataLayer);
 }
 
 FSlateColor FDataLayerPropertyTypeCustomization::GetForegroundColor() const
 {
 	FPropertyAccess::Result PropertyAccessResult;
-	const UDataLayer* DataLayer = GetDataLayerFromPropertyHandle(&PropertyAccessResult);
+	const UDataLayerInstance* DataLayer = GetDataLayerFromPropertyHandle(&PropertyAccessResult);
 	if (DataLayer && DataLayer->IsLocked())
 	{
 		return FSceneOutlinerCommonLabelData::DarkColor;
@@ -189,18 +189,18 @@ FSlateColor FDataLayerPropertyTypeCustomization::GetForegroundColor() const
 
 TSharedRef<SWidget> FDataLayerPropertyTypeCustomization::OnGetDataLayerMenu()
 {
-	return FDataLayerPropertyTypeCustomizationHelper::CreateDataLayerMenu([this](const UDataLayer* DataLayer) { AssignDataLayer(DataLayer); });
+	return FDataLayerPropertyTypeCustomizationHelper::CreateDataLayerMenu([this](const UDataLayerInstance* DataLayer) { AssignDataLayer(DataLayer); });
 }
 
 EVisibility FDataLayerPropertyTypeCustomization::GetSelectDataLayerVisibility() const
 {
-	const UDataLayer* DataLayer = GetDataLayerFromPropertyHandle();
+	const UDataLayerInstance* DataLayer = GetDataLayerFromPropertyHandle();
 	return DataLayer ? EVisibility::Visible : EVisibility::Collapsed;
 }
 
 FReply FDataLayerPropertyTypeCustomization::OnSelectDataLayer()
 {
-	if (UDataLayer* DataLayer = GetDataLayerFromPropertyHandle())
+	if (UDataLayerInstance* DataLayer = GetDataLayerFromPropertyHandle())
 	{
 		GEditor->SelectNone(true, true);
 		UDataLayerEditorSubsystem::Get()->SelectActorsInDataLayer(DataLayer, true, true);
@@ -208,11 +208,11 @@ FReply FDataLayerPropertyTypeCustomization::OnSelectDataLayer()
 	return FReply::Handled();
 }
 
-void FDataLayerPropertyTypeCustomization::AssignDataLayer(const UDataLayer* InDataLayer)
+void FDataLayerPropertyTypeCustomization::AssignDataLayer(const UDataLayerInstance* InDataLayerInstance)
 {
-	if (GetDataLayerFromPropertyHandle() != InDataLayer)
+	if (GetDataLayerFromPropertyHandle() != InDataLayerInstance)
 	{
-		PropertyHandle->SetValue(InDataLayer ? InDataLayer->GetFName() : NAME_None);
+		PropertyHandle->SetValue(InDataLayerInstance ? InDataLayerInstance->GetDataLayerFName() : NAME_None);
 		UDataLayerEditorSubsystem::Get()->OnDataLayerChanged().Broadcast(EDataLayerAction::Reset, NULL, NAME_None);
 	}
 }
@@ -222,10 +222,10 @@ FReply FDataLayerPropertyTypeCustomization::OnDrop(const FGeometry& InGeometry,
 	TSharedPtr<const FDataLayerDragDropOp> DataLayerDragDropOp = GetDataLayerDragDropOp(InDragDropEvent.GetOperation());
 	if (DataLayerDragDropOp.IsValid())
 	{
-		const TArray<FName>& DataLayerLabels = DataLayerDragDropOp->DataLayerLabels;
-		if (ensure(DataLayerLabels.Num() == 1))
+		const TArray<FDataLayerDragDropOp::FDragDropInfo>& DragDropInfos = DataLayerDragDropOp->DataLayerDragDropInfos;
+		if (ensure(DragDropInfos.Num() == 1))
 		{
-			if (const UDataLayer* DataLayerPtr = UDataLayerEditorSubsystem::Get()->GetDataLayerFromLabel(DataLayerLabels[0]))
+			if (const UDataLayerInstance* DataLayerPtr = UDataLayerEditorSubsystem::Get()->GetDataLayerInstance(DragDropInfos[0].DataLayerInstanceName))
 			{
 				AssignDataLayer(DataLayerPtr);
 			}
@@ -237,7 +237,7 @@ FReply FDataLayerPropertyTypeCustomization::OnDrop(const FGeometry& InGeometry,
 bool FDataLayerPropertyTypeCustomization::OnVerifyDrag(TSharedPtr<FDragDropOperation> InDragDrop)
 {
 	TSharedPtr<const FDataLayerDragDropOp> DataLayerDragDropOp = GetDataLayerDragDropOp(InDragDrop);
-	return DataLayerDragDropOp.IsValid() && DataLayerDragDropOp->DataLayerLabels.Num() == 1;
+	return DataLayerDragDropOp.IsValid() && DataLayerDragDropOp->DataLayerDragDropInfos.Num() == 1;
 }
 
 TSharedPtr<const FDataLayerDragDropOp> FDataLayerPropertyTypeCustomization::GetDataLayerDragDropOp(TSharedPtr<FDragDropOperation> InDragDrop)
diff --git a/Engine/Source/Editor/DataLayerEditor/Private/DataLayer/DataLayerPropertyTypeCustomization.h b/Engine/Source/Editor/DataLayerEditor/Private/DataLayer/DataLayerPropertyTypeCustomization.h
--- a/Engine/Source/Editor/DataLayerEditor/Private/DataLayer/DataLayerPropertyTypeCustomization.h
+++ b/Engine/Source/Editor/DataLayerEditor/Private/DataLayer/DataLayerPropertyTypeCustomization.h
@@ -11,7 +11,7 @@ class FReply;
 class SWidget;
 class FDragDropOperation;
 class FDataLayerDragDropOp;
-class UDataLayer;
+class UDataLayerInstance;
 struct FSlateColor;
 struct FSlateBrush;
 
@@ -22,9 +22,9 @@ struct FDataLayerPropertyTypeCustomization : public IPropertyTypeCustomization
 
 private:
 
-	void AssignDataLayer(const UDataLayer* InDataLayer);
+	void AssignDataLayer(const UDataLayerInstance* InDataLayer);
 
-	UDataLayer* GetDataLayerFromPropertyHandle(FPropertyAccess::Result* OutPropertyAccessResult = nullptr) const;
+	UDataLayerInstance* GetDataLayerFromPropertyHandle(FPropertyAccess::Result* OutPropertyAccessResult = nullptr) const;
 	FText GetDataLayerText() const;
 	FSlateColor GetForegroundColor() const;
 	const FSlateBrush* GetDataLayerIcon() const;
diff --git a/Engine/Source/Editor/DataLayerEditor/Private/DataLayer/DataLayerPropertyTypeCustomizationHelper.cpp b/Engine/Source/Editor/DataLayerEditor/Private/DataLayer/DataLayerPropertyTypeCustomizationHelper.cpp
--- a/Engine/Source/Editor/DataLayerEditor/Private/DataLayer/DataLayerPropertyTypeCustomizationHelper.cpp
+++ b/Engine/Source/Editor/DataLayerEditor/Private/DataLayer/DataLayerPropertyTypeCustomizationHelper.cpp
@@ -1,14 +1,14 @@
 // Copyright Epic Games, Inc. All Rights Reserved.
 
 #include "DataLayer/DataLayerPropertyTypeCustomizationHelper.h"
-#include "WorldPartition/DataLayer/DataLayer.h"
+#include "WorldPartition/DataLayer/DataLayerInstance.h"
 #include "DataLayerMode.h"
 
 #define LOCTEXT_NAMESPACE "DataLayer"
 
-TSharedRef<SWidget> FDataLayerPropertyTypeCustomizationHelper::CreateDataLayerMenu(TFunction<void(const UDataLayer* DataLayer)> OnDataLayerSelectedFunction)
+TSharedRef<SWidget> FDataLayerPropertyTypeCustomizationHelper::CreateDataLayerMenu(TFunction<void(const UDataLayerInstance* DataLayer)> OnDataLayerSelectedFunction)
 {
-	return FDataLayerPickingMode::CreateDataLayerPickerWidget(FOnDataLayerPicked::CreateLambda([OnDataLayerSelectedFunction](UDataLayer* TargetDataLayer)
+	return FDataLayerPickingMode::CreateDataLayerPickerWidget(FOnDataLayerPicked::CreateLambda([OnDataLayerSelectedFunction](UDataLayerInstance* TargetDataLayer)
 	{
 		OnDataLayerSelectedFunction(TargetDataLayer);
 	}));
diff --git a/Engine/Source/Editor/DataLayerEditor/Private/DataLayer/DataLayerTreeItem.cpp b/Engine/Source/Editor/DataLayerEditor/Private/DataLayer/DataLayerTreeItem.cpp
--- a/Engine/Source/Editor/DataLayerEditor/Private/DataLayer/DataLayerTreeItem.cpp
+++ b/Engine/Source/Editor/DataLayerEditor/Private/DataLayer/DataLayerTreeItem.cpp
@@ -11,10 +11,10 @@
 
 const FSceneOutlinerTreeItemType FDataLayerTreeItem::Type(&ISceneOutlinerTreeItem::Type);
 
-FDataLayerTreeItem::FDataLayerTreeItem(UDataLayer* InDataLayer)
+FDataLayerTreeItem::FDataLayerTreeItem(UDataLayerInstance* InDataLayerInstance)
 	: ISceneOutlinerTreeItem(Type)
-	, DataLayer(InDataLayer)
-	, ID(InDataLayer)
+	, DataLayerInstance(InDataLayerInstance)
+	, ID(InDataLayerInstance)
 	, bIsHighlighedtIfSelected(false)
 {
 	Flags.bIsExpanded = false;
@@ -22,16 +22,21 @@ FDataLayerTreeItem::FDataLayerTreeItem(UDataLayer* InDataLayer)
 
 FString FDataLayerTreeItem::GetDisplayString() const
 {
-	const UDataLayer* DataLayerPtr = DataLayer.Get();
-	return DataLayerPtr ? DataLayerPtr->GetDataLayerLabel().ToString() : LOCTEXT("DataLayerForMissingDataLayer", "(Deleted Data Layer)").ToString();
+	const UDataLayerInstance* DataLayerPtr = DataLayerInstance.Get();
+	return DataLayerPtr ? DataLayerPtr->GetDataLayerShortName() : LOCTEXT("DataLayerForMissingDataLayer", "(Deleted Data Layer)").ToString();
 }
 
 bool FDataLayerTreeItem::GetVisibility() const
 {
-	const UDataLayer* DataLayerPtr = DataLayer.Get();
+	const UDataLayerInstance* DataLayerPtr = DataLayerInstance.Get();
 	return DataLayerPtr && DataLayerPtr->IsVisible();
 }
 
+bool FDataLayerTreeItem::CanInteract() const 
+{
+	return true;
+}
+
 TSharedRef<SWidget> FDataLayerTreeItem::GenerateLabelWidget(ISceneOutliner& Outliner, const STableRow<FSceneOutlinerTreeItemPtr>& InRow)
 {
 	return SNew(SDataLayerTreeLabel, *this, Outliner, InRow);
@@ -39,7 +44,7 @@ TSharedRef<SWidget> FDataLayerTreeItem::GenerateLabelWidget(ISceneOutliner& Outl
 
 void FDataLayerTreeItem::OnVisibilityChanged(const bool bNewVisibility)
 {
-	if (UDataLayer* DataLayerPtr = DataLayer.Get())
+	if (UDataLayerInstance* DataLayerPtr = DataLayerInstance.Get())
 	{
 		UDataLayerEditorSubsystem::Get()->SetDataLayerVisibility(DataLayerPtr, bNewVisibility);
 	}
@@ -49,7 +54,7 @@ bool FDataLayerTreeItem::ShouldBeHighlighted() const
 {
 	if (bIsHighlighedtIfSelected)
 	{
-		if (UDataLayer* DataLayerPtr = DataLayer.Get())
+		if (UDataLayerInstance* DataLayerPtr = DataLayerInstance.Get())
 		{
 			return UDataLayerEditorSubsystem::Get()->DoesDataLayerContainSelectedActors(DataLayerPtr);
 		}
diff --git a/Engine/Source/Editor/DataLayerEditor/Private/DataLayer/DataLayerTreeItem.h b/Engine/Source/Editor/DataLayerEditor/Private/DataLayer/DataLayerTreeItem.h
--- a/Engine/Source/Editor/DataLayerEditor/Private/DataLayer/DataLayerTreeItem.h
+++ b/Engine/Source/Editor/DataLayerEditor/Private/DataLayer/DataLayerTreeItem.h
@@ -6,19 +6,20 @@
 #include "Misc/Attribute.h"
 #include "SceneOutlinerFwd.h"
 #include "ISceneOutlinerTreeItem.h"
-#include "WorldPartition/DataLayer/DataLayer.h"
+
+class UDataLayerInstance;
 
 struct FDataLayerTreeItem : ISceneOutlinerTreeItem
 {
 public:
-	FDataLayerTreeItem(UDataLayer* InDataLayer);
-	UDataLayer* GetDataLayer() const { return DataLayer.Get(); }
+	FDataLayerTreeItem(UDataLayerInstance* InDataLayerInstance);
+	UDataLayerInstance* GetDataLayer() const { return DataLayerInstance.Get(); }
 
 	/* Begin ISceneOutlinerTreeItem Implementation */
-	virtual bool IsValid() const override { return DataLayer.IsValid(); }
+	virtual bool IsValid() const override { return DataLayerInstance.IsValid(); }
 	virtual FSceneOutlinerTreeItemID GetID() const override { return ID; }
 	virtual FString GetDisplayString() const override;
-	virtual bool CanInteract() const override { return true; }
+	virtual bool CanInteract() const override;
 	virtual TSharedRef<SWidget> GenerateLabelWidget(ISceneOutliner& Outliner, const STableRow<FSceneOutlinerTreeItemPtr>& InRow) override;
 	virtual bool HasVisibilityInfo() const override { return true; }
 	virtual void OnVisibilityChanged(const bool bNewVisibility) override;
@@ -31,8 +32,8 @@ public:
 
 	static const FSceneOutlinerTreeItemType Type;
 
-	DECLARE_DELEGATE_RetVal_OneParam(bool, FFilterPredicate, const UDataLayer*);
-	DECLARE_DELEGATE_RetVal_OneParam(bool, FInteractivePredicate, const UDataLayer*);
+	DECLARE_DELEGATE_RetVal_OneParam(bool, FFilterPredicate, const UDataLayerInstance*);
+	DECLARE_DELEGATE_RetVal_OneParam(bool, FInteractivePredicate, const UDataLayerInstance*);
 
 	bool Filter(FFilterPredicate Pred) const
 	{
@@ -45,7 +46,7 @@ public:
 	}
 
 private:
-	TWeakObjectPtr<UDataLayer> DataLayer;
+	TWeakObjectPtr<UDataLayerInstance> DataLayerInstance;
 	const FObjectKey ID;
 	bool bIsHighlighedtIfSelected;
 };
\ No newline at end of file
diff --git a/Engine/Source/Editor/DataLayerEditor/Private/DataLayer/DataLayersActorDescTreeItem.h b/Engine/Source/Editor/DataLayerEditor/Private/DataLayer/DataLayersActorDescTreeItem.h
--- a/Engine/Source/Editor/DataLayerEditor/Private/DataLayer/DataLayersActorDescTreeItem.h
+++ b/Engine/Source/Editor/DataLayerEditor/Private/DataLayer/DataLayersActorDescTreeItem.h
@@ -5,7 +5,7 @@
 #include "CoreMinimal.h"
 #include "SceneOutlinerFwd.h"
 #include "ActorDescTreeItem.h"
-#include "WorldPartition/DataLayer/DataLayer.h"
+#include "WorldPartition/DataLayer/DataLayerInstance.h"
 
 //////////////////////////////////////////////////////////////////////////
 // FDataLayerActorTreeItemData
@@ -14,7 +14,7 @@ class FWorldPartitionActorDesc;
 
 struct FDataLayerActorDescTreeItemData
 {
-	FDataLayerActorDescTreeItemData(const FGuid& InActorGuid, UActorDescContainer* InContainer, UDataLayer* InDataLayer)
+	FDataLayerActorDescTreeItemData(const FGuid& InActorGuid, UActorDescContainer* InContainer, UDataLayerInstance* InDataLayer)
 		: ActorGuid(InActorGuid)
 		, Container(InContainer)
 		, DataLayer(InDataLayer)
@@ -22,7 +22,7 @@ struct FDataLayerActorDescTreeItemData
 	
 	const FGuid& ActorGuid;
 	UActorDescContainer* const Container;
-	TWeakObjectPtr<UDataLayer> DataLayer;
+	TWeakObjectPtr<UDataLayerInstance> DataLayer;
 };
 
 //////////////////////////////////////////////////////////////////////////
@@ -31,8 +31,8 @@ struct FDataLayerActorDescTreeItemData
 struct FDataLayerActorDescTreeItem : public FActorDescTreeItem
 {
 public:
-	DECLARE_DELEGATE_RetVal_TwoParams(bool, FFilterPredicate, const FWorldPartitionActorDesc* ActorDesc, const UDataLayer* DataLayer);
-	DECLARE_DELEGATE_RetVal_TwoParams(bool, FInteractivePredicate, const FWorldPartitionActorDesc* ActorDesc, const UDataLayer* DataLayer);
+	DECLARE_DELEGATE_RetVal_TwoParams(bool, FFilterPredicate, const FWorldPartitionActorDesc* ActorDesc, const UDataLayerInstance* DataLayer);
+	DECLARE_DELEGATE_RetVal_TwoParams(bool, FInteractivePredicate, const FWorldPartitionActorDesc* ActorDesc, const UDataLayerInstance* DataLayer);
 
 	FDataLayerActorDescTreeItem(const FDataLayerActorDescTreeItemData& InData)
 		: FActorDescTreeItem(InData.ActorGuid, InData.Container)
@@ -41,9 +41,9 @@ public:
 	{
 	}
 
-	UDataLayer* GetDataLayer() const { return DataLayer.Get(); }
+	UDataLayerInstance* GetDataLayer() const { return DataLayer.Get(); }
 	
-	static uint32 ComputeTreeItemID(const FGuid& InActorGuid, const UDataLayer* InDataLayer)
+	static uint32 ComputeTreeItemID(const FGuid& InActorGuid, const UDataLayerInstance* InDataLayer)
 	{
 		return HashCombine(GetTypeHash(InActorGuid), GetTypeHash(FObjectKey(InDataLayer)));
 	}
@@ -68,6 +68,6 @@ public:
 	/* End ISceneOutlinerTreeItem Implementation */
 
 private:
-	TWeakObjectPtr<UDataLayer> DataLayer;
+	TWeakObjectPtr<UDataLayerInstance> DataLayer;
 	const uint32 IDDataLayerActorDesc;
 };
\ No newline at end of file
diff --git a/Engine/Source/Editor/DataLayerEditor/Private/DataLayer/SDataLayerBrowser.cpp b/Engine/Source/Editor/DataLayerEditor/Private/DataLayer/SDataLayerBrowser.cpp
--- a/Engine/Source/Editor/DataLayerEditor/Private/DataLayer/SDataLayerBrowser.cpp
+++ b/Engine/Source/Editor/DataLayerEditor/Private/DataLayer/SDataLayerBrowser.cpp
@@ -9,17 +9,21 @@
 #include "DataLayerTreeItem.h"
 #include "DataLayerMode.h"
 #include "WorldPartition/WorldPartitionEditorPerProjectUserSettings.h"
-#include "WorldPartition/DataLayer/DataLayer.h"
+#include "WorldPartition/DataLayer/DataLayerInstance.h"
 #include "DataLayerOutlinerIsVisibleColumn.h"
 #include "DataLayerOutlinerIsLoadedInEditorColumn.h"
 #include "DataLayerOutlinerDeleteButtonColumn.h"
+#include "DataLayerOutlinerHasErrorColumn.h"
 #include "DataLayer/DataLayerEditorSubsystem.h"
-#include "WorldPartition/DataLayer/DataLayer.h"
+#include "WorldPartition/DataLayer/DataLayerInstance.h"
 #include "WorldPartition/WorldPartitionSubsystem.h"
 #include "Modules/ModuleManager.h"
 #include "ScopedTransaction.h"
 #include "SceneOutlinerTextInfoColumn.h"
 #include "Editor.h"
+#include "Widgets/Input/SMultiLineEditableTextBox.h"
+#include "DetailLayoutBuilder.h"
+#include "WorldPartition/DataLayer/WorldDataLayers.h"
 
 #define LOCTEXT_NAMESPACE "DataLayer"
 
@@ -50,9 +54,9 @@ void SDataLayerBrowser::Construct(const FArguments& InArgs)
 	{
 		if (const FDataLayerTreeItem* DataLayerItem = Item.CastTo<FDataLayerTreeItem>())
 		{
-			if (const UDataLayer* DataLayer = DataLayerItem->GetDataLayer())
+			if (const UDataLayerInstance* DataLayerInstance = DataLayerItem->GetDataLayer())
 			{
-				return DataLayer->GetFName().ToString();
+				return DataLayerInstance->GetDataLayerFName().ToString();
 			}
 		}
 		else if (const FDataLayerActorTreeItem* DataLayerActorTreeItem = Item.CastTo<FDataLayerActorTreeItem>())
@@ -76,17 +80,25 @@ void SDataLayerBrowser::Construct(const FArguments& InArgs)
 	{
 		if (const FDataLayerTreeItem* DataLayerItem = Item.CastTo<FDataLayerTreeItem>())
 		{
-			if (const UDataLayer* DataLayer = DataLayerItem->GetDataLayer())
+			if (const UDataLayerInstance* DataLayerInstance = DataLayerItem->GetDataLayer())
 			{
-				if (DataLayer->IsRuntime())
+				if (DataLayerInstance->IsRuntime())
 				{
-					return GetDataLayerRuntimeStateName(DataLayer->GetInitialRuntimeState());
+					return GetDataLayerRuntimeStateName(DataLayerInstance->GetInitialRuntimeState());
 				}
 			}
 		}
 		return FString();
 	});
 
+	SAssignNew(DeprecatedDataLayerWarningBox, SMultiLineEditableTextBox)
+		.IsReadOnly(true)
+		.Font(IDetailLayoutBuilder::GetDetailFontBold())
+		.BackgroundColor(FEditorStyle::GetColor("ErrorReporting.WarningBackgroundColor"))
+		.Text(LOCTEXT("Deprecated_DataLayers", "Some data within DataLayers is deprecated. Run DataLayerToAssetCommandlet to create DataLayerInstances and DataLayer Assets for this level."))
+		.AutoWrapText(true)
+		.Visibility_Lambda([]() { return UDataLayerEditorSubsystem::Get()->HasDeprecatedDataLayers() ? EVisibility::Visible : EVisibility::Collapsed; });
+
 	FSceneOutlinerInitializationOptions InitOptions;
 	InitOptions.bShowHeaderRow = true;
 	InitOptions.bShowParentTree = true;
@@ -98,6 +110,7 @@ void SDataLayerBrowser::Construct(const FArguments& InArgs)
 	InitOptions.ColumnMap.Add(FDataLayerOutlinerDeleteButtonColumn::GetID(), FSceneOutlinerColumnInfo(ESceneOutlinerColumnVisibility::Visible, 10, FCreateSceneOutlinerColumn::CreateLambda([](ISceneOutliner& InSceneOutliner) { return MakeShareable(new FDataLayerOutlinerDeleteButtonColumn(InSceneOutliner)); })));
 	InitOptions.ColumnMap.Add("ID Name", FSceneOutlinerColumnInfo(ESceneOutlinerColumnVisibility::Invisible, 20, FCreateSceneOutlinerColumn::CreateStatic(&FTextInfoColumn::CreateTextInfoColumn, FName("ID Name"), InternalNameInfoText, FText::GetEmpty())));
 	InitOptions.ColumnMap.Add("Initial State", FSceneOutlinerColumnInfo(ESceneOutlinerColumnVisibility::Invisible, 20, FCreateSceneOutlinerColumn::CreateStatic(&FTextInfoColumn::CreateTextInfoColumn, FName("Initial State"), InternalInitialRuntimeStateInfoText, FText::FromString("Initial Runtime State"))));
+	InitOptions.ColumnMap.Add(FDataLayerOutlinerHasErrorsColumn::GetID(), FSceneOutlinerColumnInfo(ESceneOutlinerColumnVisibility::Visible, 100, FCreateSceneOutlinerColumn::CreateLambda([](ISceneOutliner& InSceneOutliner) { return MakeShareable(new FDataLayerOutlinerHasErrorsColumn(InSceneOutliner)); }), false));
 	DataLayerOutliner = SNew(SDataLayerOutliner, InitOptions).IsEnabled(FSlateApplication::Get().GetNormalExecutionAttribute());
 
 	SAssignNew(DataLayerContentsSection, SBorder)
@@ -105,26 +118,35 @@ void SDataLayerBrowser::Construct(const FArguments& InArgs)
 	.BorderImage(FEditorStyle::GetBrush("NoBrush"))
 	.Content()
 	[
-		// Data Layer Outliner
-		SNew(SSplitter)
-		.Orientation(Orient_Vertical)
-		.Style(FEditorStyle::Get(), "DetailsView.Splitter")
-		+ SSplitter::Slot()
+		SNew(SVerticalBox)
+		+ SVerticalBox::Slot()
+		.AutoHeight()
 		[
-			SNew(SVerticalBox)
-			+ SVerticalBox::Slot()
-			[
-				DataLayerOutliner.ToSharedRef()
-			]
+			DeprecatedDataLayerWarningBox.ToSharedRef()
 		]
-		// Details
-		+SSplitter::Slot()
+		+ SVerticalBox::Slot()
 		[
-			SNew(SVerticalBox)
-			+ SVerticalBox::Slot()
-			.Padding(2, 4, 0, 0)
+			// Data Layer Outliner
+			SNew(SSplitter)
+			.Orientation(Orient_Vertical)
+			.Style(FEditorStyle::Get(), "DetailsView.Splitter")
+			+ SSplitter::Slot()
+			[
+				SNew(SVerticalBox)
+				+ SVerticalBox::Slot()
+				[
+					DataLayerOutliner.ToSharedRef()
+				]
+			]
+			// Details
+			+ SSplitter::Slot()
 			[
-				DetailsWidget.ToSharedRef()
+				SNew(SVerticalBox)
+				+ SVerticalBox::Slot()
+				.Padding(2, 4, 0, 0)
+				[
+					DetailsWidget.ToSharedRef()
+				]
 			]
 		]
 	];
@@ -134,13 +156,13 @@ void SDataLayerBrowser::Construct(const FArguments& InArgs)
 	ChildSlot
 	[
 		SAssignNew(ContentAreaBox, SVerticalBox)
-		.IsEnabled_Lambda([]() { return UWorld::IsPartitionedWorld(GWorld); })
+		.IsEnabled_Lambda([]() { return GWorld ? UWorld::IsPartitionedWorld(GWorld) : false; })
 	];
 
 	InitializeDataLayerBrowser();
 }
 
-void SDataLayerBrowser::SyncDataLayerBrowserToDataLayer(const UDataLayer* DataLayer)
+void SDataLayerBrowser::SyncDataLayerBrowserToDataLayer(const UDataLayerInstance* DataLayer)
 {
 	FSceneOutlinerTreeItemPtr Item = DataLayerOutliner->GetTreeItem(DataLayer);
 	if (Item.IsValid())
@@ -155,7 +177,7 @@ void SDataLayerBrowser::SyncDataLayerBrowserToDataLayer(const UDataLayer* DataLa
 	}
 }
 
-void SDataLayerBrowser::OnSelectionChanged(TSet<TWeakObjectPtr<const UDataLayer>>& InSelectedDataLayersSet)
+void SDataLayerBrowser::OnSelectionChanged(TSet<TWeakObjectPtr<const UDataLayerInstance>>& InSelectedDataLayersSet)
 {
 	SelectedDataLayersSet = InSelectedDataLayersSet;
 	TArray<UObject*> SelectedDataLayers;
@@ -163,7 +185,7 @@ void SDataLayerBrowser::OnSelectionChanged(TSet<TWeakObjectPtr<const UDataLayer>
 	{
 		if (WeakDataLayer.IsValid())
 		{
-			UDataLayer* DataLayer = const_cast<UDataLayer*>(WeakDataLayer.Get());
+			UDataLayerInstance* DataLayer = const_cast<UDataLayerInstance*>(WeakDataLayer.Get());
 			SelectedDataLayers.Add(DataLayer);
 		}
 	}
diff --git a/Engine/Source/Editor/DataLayerEditor/Private/DataLayer/SDataLayerBrowser.h b/Engine/Source/Editor/DataLayerEditor/Private/DataLayer/SDataLayerBrowser.h
--- a/Engine/Source/Editor/DataLayerEditor/Private/DataLayer/SDataLayerBrowser.h
+++ b/Engine/Source/Editor/DataLayerEditor/Private/DataLayer/SDataLayerBrowser.h
@@ -2,6 +2,7 @@
 #pragma once
 
 #include "CoreMinimal.h"
+#include "DataLayer/DataLayerAction.h"
 #include "InputCoreTypes.h"
 #include "Styling/SlateColor.h"
 #include "Input/Reply.h"
@@ -13,7 +14,8 @@ class ISceneOutliner;
 class SButton;
 class SVerticalBox;
 class SBorder;
-class UDataLayer;
+class SMultiLineEditableTextBox;
+class UDataLayerInstance;
 
 //////////////////////////////////////////////////////////////////////////
 // SDataLayerBrowser
@@ -25,8 +27,6 @@ public:
 	SLATE_BEGIN_ARGS( SDataLayerBrowser ) {}
 	SLATE_END_ARGS()
 
-	~SDataLayerBrowser(){}
-
 	/**
 	 * Construct this widget.  Called by the SNew() Slate macro.
 	 *
@@ -35,18 +35,18 @@ public:
 	 */
 	void Construct(const FArguments& InArgs);
 
-	void SyncDataLayerBrowserToDataLayer(const UDataLayer* DataLayer);
-	void OnSelectionChanged(TSet<TWeakObjectPtr<const UDataLayer>>& SelectedDataLayersSet);
+	void SyncDataLayerBrowserToDataLayer(const UDataLayerInstance* DataLayer);
+	void OnSelectionChanged(TSet<TWeakObjectPtr<const UDataLayerInstance>>& SelectedDataLayersSet);
 
 private:
-
 	void InitializeDataLayerBrowser();
 
-	TSet<TWeakObjectPtr<const UDataLayer>> SelectedDataLayersSet;
+	TSet<TWeakObjectPtr<const UDataLayerInstance>> SelectedDataLayersSet;
 	TSharedPtr<class SDataLayerOutliner> DataLayerOutliner;
 	TSharedPtr<class IDetailsView> DetailsWidget;
 	TSharedPtr<SButton> ToggleModeButton;
 	TSharedPtr<SVerticalBox> ContentAreaBox;
 	TSharedPtr<SBorder> DataLayerContentsSection;
 	TSharedPtr<SBorder> DataLayerContentsHeader;
+	TSharedPtr<SMultiLineEditableTextBox> DeprecatedDataLayerWarningBox;
 };
diff --git a/Engine/Source/Editor/DataLayerEditor/Private/DataLayer/SDataLayerOutliner.cpp b/Engine/Source/Editor/DataLayerEditor/Private/DataLayer/SDataLayerOutliner.cpp
--- a/Engine/Source/Editor/DataLayerEditor/Private/DataLayer/SDataLayerOutliner.cpp
+++ b/Engine/Source/Editor/DataLayerEditor/Private/DataLayer/SDataLayerOutliner.cpp
@@ -7,7 +7,7 @@
 #include "Widgets/Images/SImage.h"
 #include "DataLayer/DataLayerTreeItem.h"
 #include "DataLayer/DataLayerEditorSubsystem.h"
-#include "WorldPartition/DataLayer/DataLayer.h"
+#include "WorldPartition/DataLayer/DataLayerInstance.h"
 #include "DataLayerTransaction.h"
 #include "Algo/Transform.h"
 #include "Editor.h"
@@ -50,12 +50,12 @@ void SDataLayerOutliner::CustomAddToToolbar(TSharedPtr<SHorizontalBox> Toolbar)
 		];
 }
 
-TArray<UDataLayer*> SDataLayerOutliner::GetSelectedDataLayers() const
+TArray<UDataLayerInstance*> SDataLayerOutliner::GetSelectedDataLayers() const
 {
 	FSceneOutlinerItemSelection ItemSelection(GetSelection());
 	TArray<FDataLayerTreeItem*> SelectedDataLayerItems;
 	ItemSelection.Get<FDataLayerTreeItem>(SelectedDataLayerItems);
-	TArray<UDataLayer*> ValidSelectedDataLayers;
+	TArray<UDataLayerInstance*> ValidSelectedDataLayers;
 	Algo::TransformIf(SelectedDataLayerItems, ValidSelectedDataLayers, [](const auto Item) { return Item && Item->GetDataLayer(); }, [](const auto Item) { return Item->GetDataLayer(); });
 	return ValidSelectedDataLayers;
 }
@@ -64,8 +64,8 @@ bool SDataLayerOutliner::CanAddSelectedActorsToSelectedDataLayersClicked() const
 {
 	if (GEditor->GetSelectedActorCount() > 0)
 	{
-		TArray<UDataLayer*> SelectedDataLayers = GetSelectedDataLayers();
-		const bool bSelectedDataLayersContainsLocked = !!SelectedDataLayers.FindByPredicate([](const UDataLayer* DataLayer) { return DataLayer->IsLocked(); });
+		TArray<UDataLayerInstance*> SelectedDataLayers = GetSelectedDataLayers();
+		const bool bSelectedDataLayersContainsLocked = !!SelectedDataLayers.FindByPredicate([](const UDataLayerInstance* DataLayer) { return DataLayer->IsLocked(); });
 		return (!SelectedDataLayers.IsEmpty() && !bSelectedDataLayersContainsLocked);
 	}
 	return false;
@@ -80,7 +80,7 @@ FReply SDataLayerOutliner::OnAddSelectedActorsToSelectedDataLayersClicked()
 {
 	if (CanAddSelectedActorsToSelectedDataLayersClicked())
 	{
-		TArray<UDataLayer*> SelectedDataLayers = GetSelectedDataLayers();
+		TArray<UDataLayerInstance*> SelectedDataLayers = GetSelectedDataLayers();
 		const FScopedDataLayerTransaction Transaction(LOCTEXT("AddSelectedActorsToSelectedDataLayers", "Add Selected Actors to Selected Data Layers"), SelectedDataLayers[0]->GetWorld());
 		UDataLayerEditorSubsystem::Get()->AddSelectedActorsToDataLayers(SelectedDataLayers);
 	}
@@ -91,7 +91,7 @@ FReply SDataLayerOutliner::OnRemoveSelectedActorsFromSelectedDataLayersClicked()
 {
 	if (CanRemoveSelectedActorsFromSelectedDataLayersClicked())
 	{
-		TArray<UDataLayer*> SelectedDataLayers = GetSelectedDataLayers();
+		TArray<UDataLayerInstance*> SelectedDataLayers = GetSelectedDataLayers();
 		const FScopedDataLayerTransaction Transaction(LOCTEXT("RemoveSelectedActorsFromSelectedDataLayers", "Remove Selected Actors from Selected Data Layers"), SelectedDataLayers[0]->GetWorld());
 		UDataLayerEditorSubsystem::Get()->RemoveSelectedActorsFromDataLayers(SelectedDataLayers);
 	}
diff --git a/Engine/Source/Editor/DataLayerEditor/Private/DataLayer/SDataLayerOutliner.h b/Engine/Source/Editor/DataLayerEditor/Private/DataLayer/SDataLayerOutliner.h
--- a/Engine/Source/Editor/DataLayerEditor/Private/DataLayer/SDataLayerOutliner.h
+++ b/Engine/Source/Editor/DataLayerEditor/Private/DataLayer/SDataLayerOutliner.h
@@ -22,5 +22,5 @@ private:
 	bool CanRemoveSelectedActorsFromSelectedDataLayersClicked() const;
 	FReply OnAddSelectedActorsToSelectedDataLayersClicked();
 	FReply OnRemoveSelectedActorsFromSelectedDataLayersClicked();
-	TArray<class UDataLayer*> GetSelectedDataLayers() const;
+	TArray<class UDataLayerInstance*> GetSelectedDataLayers() const;
 };
\ No newline at end of file
diff --git a/Engine/Source/Editor/DataLayerEditor/Private/DataLayer/SDataLayerTreeLabel.cpp b/Engine/Source/Editor/DataLayerEditor/Private/DataLayer/SDataLayerTreeLabel.cpp
--- a/Engine/Source/Editor/DataLayerEditor/Private/DataLayer/SDataLayerTreeLabel.cpp
+++ b/Engine/Source/Editor/DataLayerEditor/Private/DataLayer/SDataLayerTreeLabel.cpp
@@ -38,7 +38,7 @@ void SDataLayerTreeLabel::Construct(const FArguments& InArgs, FDataLayerTreeItem
 		.IsSelected(FIsSelected::CreateSP(&InRow, &STableRow<FSceneOutlinerTreeItemPtr>::IsSelectedExclusively))
 		.IsReadOnly_Lambda([Item = DataLayerItem.AsShared(), this]()
 		{
-			return !CanExecuteRenameRequest(Item.Get());
+			return DataLayerPtr == nullptr || !DataLayerPtr->SupportRelabeling() || !CanExecuteRenameRequest(Item.Get());
 		})
 	]
 
@@ -122,7 +122,7 @@ bool SDataLayerTreeLabel::ShouldBeHighlighted() const
 
 bool SDataLayerTreeLabel::IsInActorEditorContext() const
 {
-	const UDataLayer* DataLayer = DataLayerPtr.Get();
+	const UDataLayerInstance* DataLayer = DataLayerPtr.Get();
 	return DataLayer && DataLayer->IsInActorEditorContext();
 }
 
@@ -140,7 +140,7 @@ FSlateFontInfo SDataLayerTreeLabel::GetDisplayNameFont() const
 
 FText SDataLayerTreeLabel::GetDisplayText() const
 {
-	const UDataLayer* DataLayer = DataLayerPtr.Get();
+	const UDataLayerInstance* DataLayer = DataLayerPtr.Get();
 	bool bIsDataLayerActive = false;
 	FText DataLayerRuntimeStateText = FText::GetEmpty();
 	if (DataLayer && DataLayer->IsRuntime() && DataLayer->GetWorld()->IsPlayInEditor())
@@ -150,7 +150,7 @@ FText SDataLayerTreeLabel::GetDisplayText() const
 	}
 	
 	static const FText DataLayerDeleted = LOCTEXT("DataLayerLabelForMissingDataLayer", "(Deleted Data Layer)");
-	return DataLayer ? FText::Format(LOCTEXT("DataLayerDisplayText", "{0}{1}"), FText::FromName(DataLayer->GetDataLayerLabel()), DataLayerRuntimeStateText) : DataLayerDeleted;
+	return DataLayer ? FText::Format(LOCTEXT("DataLayerDisplayText", "{0}{1}"), FText::FromString(DataLayer->GetDataLayerShortName()), DataLayerRuntimeStateText) : DataLayerDeleted;
 }
 
 FText SDataLayerTreeLabel::GetTooltipText() const
@@ -165,7 +165,7 @@ FText SDataLayerTreeLabel::GetTooltipText() const
 
 FText SDataLayerTreeLabel::GetTypeText() const
 {
-	const UDataLayer* DataLayer = DataLayerPtr.Get();
+	const UDataLayerInstance* DataLayer = DataLayerPtr.Get();
 	return DataLayer ? FText::FromName(DataLayer->GetClass()->GetFName()) : FText();
 }
 
@@ -176,7 +176,7 @@ EVisibility SDataLayerTreeLabel::GetTypeTextVisibility() const
 
 const FSlateBrush* SDataLayerTreeLabel::GetIcon() const
 {
-	const UDataLayer* DataLayer = DataLayerPtr.Get();
+	const UDataLayerInstance* DataLayer = DataLayerPtr.Get();
 	if (DataLayer && WeakSceneOutliner.IsValid())
 	{
 		const TCHAR* IconName = DataLayer->GetDataLayerIconName();
@@ -194,7 +194,7 @@ const FSlateBrush* SDataLayerTreeLabel::GetIcon() const
 
 FText SDataLayerTreeLabel::GetIconTooltip() const
 {
-	const UDataLayer* DataLayer = DataLayerPtr.Get();
+	const UDataLayerInstance* DataLayer = DataLayerPtr.Get();
 	return DataLayer ? (DataLayer->IsRuntime() ? FText(LOCTEXT("RuntimeDataLayer", "Runtime Data Layer")) : FText(LOCTEXT("EditorDataLayer", "Editor Data Layer"))) : FText();
 }
 
@@ -205,7 +205,7 @@ FSlateColor SDataLayerTreeLabel::GetForegroundColor() const
 		return BaseColor.GetValue();
 	}
 
-	const UDataLayer* DataLayer = DataLayerPtr.Get();
+	const UDataLayerInstance* DataLayer = DataLayerPtr.Get();
 	if (DataLayer)
 	{
 		if (DataLayer->GetWorld()->IsPlayInEditor())
@@ -253,23 +253,36 @@ bool SDataLayerTreeLabel::OnVerifyItemLabelChanged(const FText& InLabel, FText&
 		return false;
 	}
 
-	UDataLayer* FoundDataLayer;
-	if (UDataLayerEditorSubsystem::Get()->TryGetDataLayerFromLabel(*InLabel.ToString(), FoundDataLayer) && FoundDataLayer != DataLayerPtr.Get())
+	UDataLayerInstance* FoundDataLayerInstance;
+	
+	PRAGMA_DISABLE_DEPRECATION_WARNINGS
+	if (UDataLayerEditorSubsystem::Get()->TryGetDataLayerFromLabel(*InLabel.ToString(), FoundDataLayerInstance) && FoundDataLayerInstance != DataLayerPtr.Get())
 	{
 		OutErrorMessage = LOCTEXT("RenameFailed_AlreadyExists", "This Data Layer already exists");
 		return false;
 	}
+	PRAGMA_ENABLE_DEPRECATION_WARNINGS
+
+	if (FoundDataLayerInstance != nullptr && !FoundDataLayerInstance->SupportRelabeling())
+	{
+		OutErrorMessage = LOCTEXT("RenameFailed_NotPermittedOnDataLayer", "This Data Layer does not support renaming");
+		return false;
+	}
 
 	return true;
 }
 
 void SDataLayerTreeLabel::OnLabelCommitted(const FText& InLabel, ETextCommit::Type InCommitInfo)
 {
-	UDataLayer* DataLayer = DataLayerPtr.Get();
-	if (DataLayer && !InLabel.ToString().Equals(DataLayer->GetDataLayerLabel().ToString(), ESearchCase::CaseSensitive))
+	UDataLayerInstance* DataLayerInstance = DataLayerPtr.Get();
+	check(DataLayerInstance->SupportRelabeling());
+	if (!InLabel.ToString().Equals(DataLayerInstance->GetDataLayerShortName(), ESearchCase::CaseSensitive))
 	{
-		const FScopedDataLayerTransaction Transaction(LOCTEXT("SceneOutlinerRenameDataLayerTransaction", "Rename Data Layer"), DataLayer->GetWorld());
-		UDataLayerEditorSubsystem::Get()->RenameDataLayer(DataLayer, *InLabel.ToString());
+		const FScopedDataLayerTransaction Transaction(LOCTEXT("SceneOutlinerRenameDataLayerTransaction", "Rename Data Layer"), DataLayerInstance->GetWorld());
+
+		PRAGMA_DISABLE_DEPRECATION_WARNINGS
+		UDataLayerEditorSubsystem::Get()->RenameDataLayer(DataLayerInstance, *InLabel.ToString());
+		PRAGMA_ENABLE_DEPRECATION_WARNINGS
 
 		if (WeakSceneOutliner.IsValid())
 		{
diff --git a/Engine/Source/Editor/DataLayerEditor/Private/DataLayer/SDataLayerTreeLabel.h b/Engine/Source/Editor/DataLayerEditor/Private/DataLayer/SDataLayerTreeLabel.h
--- a/Engine/Source/Editor/DataLayerEditor/Private/DataLayer/SDataLayerTreeLabel.h
+++ b/Engine/Source/Editor/DataLayerEditor/Private/DataLayer/SDataLayerTreeLabel.h
@@ -11,7 +11,7 @@
 #include "Widgets/DeclarativeSyntaxSupport.h"
 
 struct FDataLayerTreeItem;
-class UDataLayer;
+class UDataLayerInstance;
 
 struct SDataLayerTreeLabel : FSceneOutlinerCommonLabelData, public SCompoundWidget
 {
@@ -36,6 +36,6 @@ private:
 	bool IsInActorEditorContext() const;
 
 	TWeakPtr<FDataLayerTreeItem> TreeItemPtr;
-	TWeakObjectPtr<UDataLayer> DataLayerPtr;
+	TWeakObjectPtr<UDataLayerInstance> DataLayerPtr;
 	TAttribute<FText> HighlightText;
 };
\ No newline at end of file
diff --git a/Engine/Source/Editor/DataLayerEditor/Private/DataLayerEditorModule.cpp b/Engine/Source/Editor/DataLayerEditor/Private/DataLayerEditorModule.cpp
--- a/Engine/Source/Editor/DataLayerEditor/Private/DataLayerEditorModule.cpp
+++ b/Engine/Source/Editor/DataLayerEditor/Private/DataLayerEditorModule.cpp
@@ -11,12 +11,13 @@
 #include "DataLayer/DataLayerPropertyTypeCustomization.h"
 #include "DataLayer/SDataLayerBrowser.h"
 #include "DataLayer/DataLayerNameEditSink.h"
-#include "WorldPartition/DataLayer/DataLayer.h"
+#include "WorldPartition/DataLayer/DataLayerInstance.h"
 
 IMPLEMENT_MODULE(FDataLayerEditorModule, DataLayerEditor );
 
 static const FName NAME_ActorDataLayer(TEXT("ActorDataLayer"));
 
+
 void FDataLayerEditorModule::StartupModule()
 {
 	FPropertyEditorModule& PropertyModule = FModuleManager::LoadModuleChecked<FPropertyEditorModule>("PropertyEditor");
@@ -41,11 +42,11 @@ TSharedRef<SWidget> FDataLayerEditorModule::CreateDataLayerBrowser()
 	return NewDataLayerBrowser;
 }
 
-void FDataLayerEditorModule::SyncDataLayerBrowserToDataLayer(const UDataLayer* DataLayer)
+void FDataLayerEditorModule::SyncDataLayerBrowserToDataLayer(const UDataLayerInstance* DataLayerInstance)
 {
 	if (DataLayerBrowser.IsValid())
 	{
 		TSharedRef<SDataLayerBrowser> Browser = StaticCastSharedRef<SDataLayerBrowser>(DataLayerBrowser.Pin().ToSharedRef());
-		Browser->SyncDataLayerBrowserToDataLayer(DataLayer);
+		Browser->SyncDataLayerBrowserToDataLayer(DataLayerInstance);
 	}
-}
\ No newline at end of file
+}
