UDataLayer deprecated in favor of UDataLayerInstance & UDataLayerAsset

diff --git a/Engine/Source/Editor/DataLayerEditor/Public/DataLayer/DataLayerEditorSubsystem.h b/Engine/Source/Editor/DataLayerEditor/Public/DataLayer/DataLayerEditorSubsystem.h
--- a/Engine/Source/Editor/DataLayerEditor/Public/DataLayer/DataLayerEditorSubsystem.h
+++ b/Engine/Source/Editor/DataLayerEditor/Public/DataLayer/DataLayerEditorSubsystem.h
@@ -6,7 +6,7 @@
 #include "EditorSubsystem.h"
 #include "DataLayerAction.h"
 #include "WorldPartition/DataLayer/ActorDataLayer.h"
-#include "WorldPartition/DataLayer/DataLayer.h"
+#include "WorldPartition/DataLayer/DataLayerInstance.h"
 #include "IActorEditorContextClient.h"
 #include "DataLayerEditorSubsystem.generated.h"
 
@@ -50,23 +50,23 @@ public:
 	virtual TSharedRef<SWidget> GetActorEditorContextWidget(UWorld* InWorld) const override;
 	virtual FOnActorEditorContextClientChanged& GetOnActorEditorContextClientChanged() override { return ActorEditorContextClientChanged; }
 	//~ End IActorEditorContextClient interface
-	void AddToActorEditorContext(UDataLayer* InDataLayer);
-	void RemoveFromActorEditorContext(UDataLayer* InDataLayer);
+	void AddToActorEditorContext(UDataLayerInstance* InDataLayerInstance);
+	void RemoveFromActorEditorContext(UDataLayerInstance* InDataLayerInstance);
 
 	/* Broadcasts whenever one or more DataLayers are modified
 	 *
 	 * Actions
-	 * Add    : The specified ChangedDataLayer is a newly created UDataLayer
+	 * Add    : The specified ChangedDataLayer is a newly created UDataLayerInstance
 	 * Modify : The specified ChangedDataLayer was just modified, if ChangedDataLayer is invalid then multiple DataLayers were modified.
-	 *          ChangedProperty specifies what field on the UDataLayer was changed, if NAME_None then multiple fields were changed
+	 *          ChangedProperty specifies what field on the UDataLayerInstance was changed, if NAME_None then multiple fields were changed
 	 * Delete : A DataLayer was deleted
 	 * Rename : The specified ChangedDataLayer was just renamed
 	 * Reset  : A large amount of changes have occurred to a number of DataLayers.
 	 */
-	DECLARE_EVENT_ThreeParams(UDataLayerEditorSubsystem, FOnDataLayerChanged, const EDataLayerAction /*Action*/, const TWeakObjectPtr<const UDataLayer>& /*ChangedDataLayer*/, const FName& /*ChangedProperty*/);
+	DECLARE_EVENT_ThreeParams(UDataLayerEditorSubsystem, FOnDataLayerChanged, const EDataLayerAction /*Action*/, const TWeakObjectPtr<const UDataLayerInstance>& /*ChangedDataLayer*/, const FName& /*ChangedProperty*/);
 	FOnDataLayerChanged& OnDataLayerChanged() { return DataLayerChanged; }
 
-	/** Broadcasts whenever one or more Actors changed DataLayers*/
+	/** Broadcasts whenever one or more Actors changed UDataLayerInstances*/
 	DECLARE_EVENT_OneParam(UDataLayerEditorSubsystem, FOnActorDataLayersChanged, const TWeakObjectPtr<AActor>& /*ChangedActor*/);
 	FOnActorDataLayersChanged& OnActorDataLayersChanged() { return ActorDataLayersChanged; }
 	
@@ -89,7 +89,7 @@ public:
 	 * @return				true if the actor was added.  false is returned if the actor already belongs to the DataLayer.
 	 */
 	UFUNCTION(BlueprintCallable, Category = DataLayers)
-	bool AddActorToDataLayer(AActor* Actor, UDataLayer* DataLayer);
+	bool AddActorToDataLayer(AActor* Actor, UDataLayerInstance* DataLayer);
 
 	/**
 	 * Adds the provided actor to the DataLayers.
@@ -99,7 +99,7 @@ public:
 	 * @return				true if the actor was added to at least one of the provided DataLayers.
 	 */
 	UFUNCTION(BlueprintCallable, Category = DataLayers)
-	bool AddActorToDataLayers(AActor* Actor, const TArray<UDataLayer*>& DataLayers);
+	bool AddActorToDataLayers(AActor* Actor, const TArray<UDataLayerInstance*>& DataLayers);
 
 	/**
 	 * Removes an actor from the specified DataLayer.
@@ -109,7 +109,7 @@ public:
 	 * @return					true if the actor was removed from the DataLayer.  false is returned if the actor already belonged to the DataLayer.
 	 */
 	UFUNCTION(BlueprintCallable, Category = DataLayers)
-	bool RemoveActorFromDataLayer(AActor* Actor, UDataLayer* DataLayerToRemove);
+	bool RemoveActorFromDataLayer(AActor* Actor, UDataLayerInstance* DataLayerToRemove);
 
 	/**
 	 * Removes the provided actor from the DataLayers.
@@ -119,7 +119,7 @@ public:
 	 * @return				true if the actor was removed from at least one of the provided DataLayers.
 	 */
 	UFUNCTION(BlueprintCallable, Category = DataLayers)
-	bool RemoveActorFromDataLayers(AActor* Actor, const TArray<UDataLayer*>& DataLayers);
+	bool RemoveActorFromDataLayers(AActor* Actor, const TArray<UDataLayerInstance*>& DataLayers);
 
 	/**
 	 * Removes an actor from all DataLayers.
@@ -150,7 +150,7 @@ public:
 	 * @return				true if at least one actor was added to the DataLayer.  false is returned if all the actors already belonged to the DataLayer.
 	 */
 	UFUNCTION(BlueprintCallable, Category = DataLayers)
-	bool AddActorsToDataLayer(const TArray<AActor*>& Actors, UDataLayer* DataLayer);
+	bool AddActorsToDataLayer(const TArray<AActor*>& Actors, UDataLayerInstance* DataLayer);
 
 	/**
 	 * Add the actors to the DataLayers
@@ -160,7 +160,7 @@ public:
 	 * @return				true if at least one actor was added to at least one DataLayer.  false is returned if all the actors already belonged to all specified DataLayers.
 	 */
 	UFUNCTION(BlueprintCallable, Category = DataLayers)
-	bool AddActorsToDataLayers(const TArray<AActor*>& Actors, const TArray<UDataLayer*>& DataLayers);
+	bool AddActorsToDataLayers(const TArray<AActor*>& Actors, const TArray<UDataLayerInstance*>& DataLayers);
 
 	/**
 	 * Removes the actors from the specified DataLayer.
@@ -170,7 +170,7 @@ public:
 	 * @return					true if at least one actor was removed from the DataLayer.  false is returned if all the actors already belonged to the DataLayer.
 	 */
 	UFUNCTION(BlueprintCallable, Category = DataLayers)
-	bool RemoveActorsFromDataLayer(const TArray<AActor*>& Actors, UDataLayer* DataLayer);
+	bool RemoveActorsFromDataLayer(const TArray<AActor*>& Actors, UDataLayerInstance* DataLayer);
 
 	/**
 	 * Remove the actors to the DataLayers
@@ -180,7 +180,7 @@ public:
 	 * @return				true if at least one actor was removed from at least one DataLayer.  false is returned if none of the actors belonged to any of the specified DataLayers.
 	 */
 	UFUNCTION(BlueprintCallable, Category = DataLayers)
-	bool RemoveActorsFromDataLayers(const TArray<AActor*>& Actors, const TArray<UDataLayer*>& DataLayers);
+	bool RemoveActorsFromDataLayers(const TArray<AActor*>& Actors, const TArray<UDataLayerInstance*>& DataLayers);
 
 	/////////////////////////////////////////////////
 	// Operations on selected actors.
@@ -192,7 +192,7 @@ public:
 	 * @return				true if at least one actor was added.  false is returned if all selected actors already belong to the DataLayer.
 	 */
 	UFUNCTION(BlueprintCallable, Category = DataLayers)
-	bool AddSelectedActorsToDataLayer(UDataLayer* DataLayer);
+	bool AddSelectedActorsToDataLayer(UDataLayerInstance* DataLayer);
 
 	/**
 	 * Adds selected actors to the DataLayers.
@@ -201,7 +201,7 @@ public:
 	 * @return				true if at least one actor was added.  false is returned if all selected actors already belong to the DataLayers.
 	 */
 	UFUNCTION(BlueprintCallable, Category = DataLayers)
-	bool AddSelectedActorsToDataLayers(const TArray<UDataLayer*>& DataLayers);
+	bool AddSelectedActorsToDataLayers(const TArray<UDataLayerInstance*>& DataLayers);
 
 	/**
 	 * Removes the selected actors from the DataLayer.
@@ -210,7 +210,7 @@ public:
 	 * @return				true if at least one actor was added.  false is returned if all selected actors already belong to the DataLayer.
 	 */
 	UFUNCTION(BlueprintCallable, Category = DataLayers)
-	bool RemoveSelectedActorsFromDataLayer(UDataLayer* DataLayer);
+	bool RemoveSelectedActorsFromDataLayer(UDataLayerInstance* DataLayer);
 
 	/**
 	 * Removes selected actors from the DataLayers.
@@ -219,7 +219,7 @@ public:
 	 * @return				true if at least one actor was removed.
 	 */
 	UFUNCTION(BlueprintCallable, Category = DataLayers)
-	bool RemoveSelectedActorsFromDataLayers(const TArray<UDataLayer*>& DataLayers);
+	bool RemoveSelectedActorsFromDataLayers(const TArray<UDataLayerInstance*>& DataLayers);
 
 
 	/////////////////////////////////////////////////
@@ -235,7 +235,7 @@ public:
 	 * @return									true if at least one actor was selected/deselected.
 	 */
 	UFUNCTION(BlueprintCallable, Category = DataLayers)
-	bool SelectActorsInDataLayer(UDataLayer* DataLayer, const bool bSelect, const bool bNotify, const bool bSelectEvenIfHidden = false);
+	bool SelectActorsInDataLayer(UDataLayerInstance* DataLayer, const bool bSelect, const bool bNotify, const bool bSelectEvenIfHidden = false);
 
 	/**
 	 * Selects/de-selects actors belonging to the DataLayer.
@@ -247,7 +247,7 @@ public:
 	 * @param	Filter	[optional]				Actor that don't pass the specified filter restrictions won't be selected.
 	 * @return									true if at least one actor was selected/deselected.
 	 */
-	bool SelectActorsInDataLayer(UDataLayer* DataLayer, const bool bSelect, const bool bNotify, const bool bSelectEvenIfHidden, const TSharedPtr<FActorFilter>& Filter);
+	bool SelectActorsInDataLayer(UDataLayerInstance* DataLayer, const bool bSelect, const bool bNotify, const bool bSelectEvenIfHidden, const TSharedPtr<FActorFilter>& Filter);
 
 	/**
 	 * Selects/de-selects actors belonging to the DataLayers.
@@ -259,7 +259,7 @@ public:
 	 * @return									true if at least one actor was selected/deselected.
 	 */
 	UFUNCTION(BlueprintCallable, Category = DataLayers)
-	bool SelectActorsInDataLayers(const TArray<UDataLayer*>& DataLayers, const bool bSelect, const bool bNotify, const bool bSelectEvenIfHidden = false);
+	bool SelectActorsInDataLayers(const TArray<UDataLayerInstance*>& DataLayers, const bool bSelect, const bool bNotify, const bool bSelectEvenIfHidden = false);
 
 	/**
 	 * Selects/de-selects actors belonging to the DataLayers.
@@ -271,7 +271,7 @@ public:
 	 * @param	Filter	[optional]				Actor that don't pass the specified filter restrictions won't be selected.
 	 * @return									true if at least one actor was selected/deselected.
 	 */
-	bool SelectActorsInDataLayers(const TArray<UDataLayer*>& DataLayers, const bool bSelect, const bool bNotify, const bool bSelectEvenIfHidden, const TSharedPtr<FActorFilter>& Filter);
+	bool SelectActorsInDataLayers(const TArray<UDataLayerInstance*>& DataLayers, const bool bSelect, const bool bNotify, const bool bSelectEvenIfHidden, const TSharedPtr<FActorFilter>& Filter);
 
 
 	/////////////////////////////////////////////////
@@ -311,7 +311,7 @@ public:
 	 *	@param	InOutActors			The list to append the found actors to.
 	 */
 	UFUNCTION(BlueprintCallable, Category = DataLayers)
-	void AppendActorsFromDataLayer(UDataLayer* DataLayer, TArray<AActor*>& InOutActors) const;
+	void AppendActorsFromDataLayer(UDataLayerInstance* DataLayer, TArray<AActor*>& InOutActors) const;
 
 	/**
 	 *	Appends all the actors associated with the specified DataLayer.
@@ -320,7 +320,7 @@ public:
 	 *	@param	InOutActors			The list to append the found actors to.
 	 *  @param	Filter	[optional]	Actor that don't pass the specified filter restrictions won't be selected.
 	 */
-	void AppendActorsFromDataLayer(UDataLayer* DataLayer, TArray<AActor*>& InOutActors, const TSharedPtr<FActorFilter>& Filter) const;
+	void AppendActorsFromDataLayer(UDataLayerInstance* DataLayer, TArray<AActor*>& InOutActors, const TSharedPtr<FActorFilter>& Filter) const;
 	
 	/**
 	 *	Appends all the actors associated with ANY of the specified DataLayers.
@@ -329,7 +329,7 @@ public:
 	 *	@param	InOutActors			The list to append the found actors to.
 	 */
 	UFUNCTION(BlueprintCallable, Category = DataLayers)
-	void AppendActorsFromDataLayers(const TArray<UDataLayer*>& DataLayers, TArray<AActor*>& InOutActors) const;
+	void AppendActorsFromDataLayers(const TArray<UDataLayerInstance*>& DataLayers, TArray<AActor*>& InOutActors) const;
 
 	/**
 	 *	Appends all the actors associated with ANY of the specified DataLayers.
@@ -338,7 +338,7 @@ public:
 	 *	@param	InOutActors			The list to append the found actors to.
 	 *  @param	Filter	[optional]	Actor that don't pass the specified filter restrictions won't be selected.
 	 */
-	void AppendActorsFromDataLayers(const TArray<UDataLayer*>& DataLayers, TArray<AActor*>& InOutActors, const TSharedPtr<FActorFilter>& Filter) const;
+	void AppendActorsFromDataLayers(const TArray<UDataLayerInstance*>& DataLayers, TArray<AActor*>& InOutActors, const TSharedPtr<FActorFilter>& Filter) const;
 
 	/**
 	 *	Gets all the actors associated with the specified DataLayer. Analog to AppendActorsFromDataLayer but it returns rather than appends the actors.
@@ -347,7 +347,7 @@ public:
 	 *	@return						The list to assign the found actors to.
 	 */
 	UFUNCTION(BlueprintCallable, Category = DataLayers)
-	TArray<AActor*> GetActorsFromDataLayer(UDataLayer* DataLayer) const;
+	TArray<AActor*> GetActorsFromDataLayer(UDataLayerInstance* DataLayer) const;
 
 	/**
 	 *	Gets all the actors associated with the specified DataLayer. Analog to AppendActorsFromDataLayer but it returns rather than appends the actors.
@@ -356,7 +356,7 @@ public:
 	 *  @param	Filter	[optional]	Actor that don't pass the specified filter restrictions won't be selected.
 	 *	@return						The list to assign the found actors to.
 	 */
-	TArray<AActor*> GetActorsFromDataLayer(UDataLayer* DataLayer, const TSharedPtr<FActorFilter>& Filter) const;
+	TArray<AActor*> GetActorsFromDataLayer(UDataLayerInstance* DataLayer, const TSharedPtr<FActorFilter>& Filter) const;
 
 	/**
 	 *	Gets all the actors associated with ANY of the specified DataLayers. Analog to AppendActorsFromDataLayers but it returns rather than appends the actors.
@@ -365,7 +365,7 @@ public:
 	 *	@return						The list to assign the found actors to.
 	 */
 	UFUNCTION(BlueprintCallable, Category = DataLayers)
-	TArray<AActor*> GetActorsFromDataLayers(const TArray<UDataLayer*>& DataLayers) const;
+	TArray<AActor*> GetActorsFromDataLayers(const TArray<UDataLayerInstance*>& DataLayers) const;
 
 	/**
 	 *	Gets all the actors associated with ANY of the specified DataLayers. Analog to AppendActorsFromDataLayers but it returns rather than appends the actors.
@@ -374,7 +374,7 @@ public:
 	 *  @param	Filter	[optional]	Actor that don't pass the specified filter restrictions won't be selected.
 	 *	@return						The list to assign the found actors to.
 	 */
-	TArray<AActor*> GetActorsFromDataLayers(const TArray<UDataLayer*>& DataLayers, const TSharedPtr<FActorFilter>& Filter) const;
+	TArray<AActor*> GetActorsFromDataLayers(const TArray<UDataLayerInstance*>& DataLayers, const TSharedPtr<FActorFilter>& Filter) const;
 
 	/**
 	 * Changes the DataLayer's visibility to the provided state
@@ -383,7 +383,7 @@ public:
 	 * @param	bIsVisible	If true the DataLayer will be visible; if false, the DataLayer will not be visible.
 	 */
 	UFUNCTION(BlueprintCallable, Category = DataLayers)
-	void SetDataLayerVisibility(UDataLayer* DataLayer, const bool bIsVisible);
+	void SetDataLayerVisibility(UDataLayerInstance* DataLayer, const bool bIsVisible);
 
 	/**
 	 * Changes visibility of the DataLayers to the provided state
@@ -392,7 +392,7 @@ public:
 	 * @param	bIsVisible	If true the DataLayers will be visible; if false, the DataLayers will not be visible
 	 */
 	UFUNCTION(BlueprintCallable, Category = DataLayers)
-	void SetDataLayersVisibility(const TArray<UDataLayer*>& DataLayers, const bool bIsVisible);
+	void SetDataLayersVisibility(const TArray<UDataLayerInstance*>& DataLayers, const bool bIsVisible);
 
 	/**
 	 * Toggles the DataLayer's visibility
@@ -400,7 +400,7 @@ public:
 	 * @param DataLayer	The DataLayer to affect
 	 */
 	UFUNCTION(BlueprintCallable, Category = DataLayers)
-	void ToggleDataLayerVisibility(UDataLayer* DataLayer);
+	void ToggleDataLayerVisibility(UDataLayerInstance* DataLayer);
 
 	/**
 	 * Toggles the visibility of all of the DataLayers
@@ -408,7 +408,7 @@ public:
 	 * @param	DataLayers	The DataLayers to affect
 	 */
 	UFUNCTION(BlueprintCallable, Category = DataLayers)
-	void ToggleDataLayersVisibility(const TArray<UDataLayer*>& DataLayers);
+	void ToggleDataLayersVisibility(const TArray<UDataLayerInstance*>& DataLayers);
 
 	/**
 	 * Changes the DataLayer's IsLoadedInEditor flag to the provided state
@@ -420,7 +420,7 @@ public:
 	 * @param	bIsFromUserChange				If this change originates from a user change or not.
 	 */
 	UFUNCTION(BlueprintCallable, Category = DataLayers)
-	bool SetDataLayerIsLoadedInEditor(UDataLayer* DataLayer, const bool bIsLoadedInEditor, const bool bIsFromUserChange);
+	bool SetDataLayerIsLoadedInEditor(UDataLayerInstance* DataLayer, const bool bIsLoadedInEditor, const bool bIsFromUserChange);
 
 	/**
 	 * Changes the IsLoadedInEditor flag of the DataLayers to the provided state
@@ -432,7 +432,7 @@ public:
 	 * @param	bIsFromUserChange				If this change originates from a user change or not.
 	 */
 	UFUNCTION(BlueprintCallable, Category = DataLayers)
-	bool SetDataLayersIsLoadedInEditor(const TArray<UDataLayer*>& DataLayers, const bool bIsLoadedInEditor, const bool bIsFromUserChange);
+	bool SetDataLayersIsLoadedInEditor(const TArray<UDataLayerInstance*>& DataLayers, const bool bIsLoadedInEditor, const bool bIsFromUserChange);
 
 	/**
 	 * Toggles the DataLayer's IsLoadedInEditor flag
@@ -441,7 +441,7 @@ public:
 	 * @param	bIsFromUserChange				If this change originates from a user change or not.
 	 */
 	UFUNCTION(BlueprintCallable, Category = DataLayers)
-	bool ToggleDataLayerIsLoadedInEditor(UDataLayer* DataLayer, const bool bIsFromUserChange);
+	bool ToggleDataLayerIsLoadedInEditor(UDataLayerInstance* DataLayer, const bool bIsFromUserChange);
 
 	/**
 	 * Toggles the IsLoadedInEditor flag of all of the DataLayers
@@ -450,7 +450,7 @@ public:
 	 * @param	bIsFromUserChange				If this change originates from a user change or not.
 	 */
 	UFUNCTION(BlueprintCallable, Category = DataLayers)
-	bool ToggleDataLayersIsLoadedInEditor(const TArray<UDataLayer*>& DataLayers, const bool bIsFromUserChange);
+	bool ToggleDataLayersIsLoadedInEditor(const TArray<UDataLayerInstance*>& DataLayers, const bool bIsFromUserChange);
 
 	/**
 	 * Set the visibility of all DataLayers to true
@@ -459,54 +459,37 @@ public:
 	void MakeAllDataLayersVisible();
 
 	/**
-	 * Gets the UDataLayer Object of the DataLayer label
+	 * Gets the UDataLayerInstance Object from associate with the DataLayerAsset
 	 *
-	 * @param	DataLayerLabel	The label of the DataLayer whose UDataLayer Object is returned
-	 * @return					The UDataLayer Object of the provided DataLayer label
+	 * @param	DataLayerAsset	The DataLayerAsset associated to the UDataLayerInstance
+	 * @return					The UDataLayerInstance Object of the provided DataLayerAsset
 	 */
 	UFUNCTION(BlueprintCallable, Category = DataLayers)
-	UDataLayer* GetDataLayerFromLabel(const FName& DataLayerLabel) const;
+	UDataLayerInstance* GetDataLayerInstance(const UDataLayerAsset* DataLayerAsset) const;
 
 	/**
-	 * Gets the UDataLayer Object of the ActorDataLayer
+	 * Gets the UDataLayerInstance Object of the DataLayer name
 	 *
-	 * @param	ActorDataLayer	The ActorDataLayer whose UDataLayer Object is returned
-	 * @return					The UDataLayer Object of the provided ActorDataLayer
+	 * @param	DataLayerName	The name of the DataLayer whose UDataLayerInstance Object is returned
+	 * @return					The UDataLayerInstance Object of the provided DataLayer name
 	 */
-	UFUNCTION(BlueprintCallable, Category = DataLayers)
-	UDataLayer* GetDataLayer(const FActorDataLayer& ActorDataLayer) const;
-
-	/**
-	 * Gets the UDataLayer Object of the DataLayer name
-	 *
-	 * @param	DataLayerName	The name of the DataLayer whose UDataLayer Object is returned
-	 * @return					The UDataLayer Object of the provided DataLayer name
-	 */
-	UDataLayer* GetDataLayerFromName(const FName& DataLayerName) const;
-	
-	/**
-	 * Attempts to get the UDataLayer Object of the provided DataLayer label.
-	 *
-	 * @param	DataLayerLabel		The label of the DataLayer whose UDataLayer Object to retrieve
-	 * @param	OutDataLayer[OUT] 	Set to the UDataLayer Object of the DataLayer label. Set to Invalid if no UDataLayer Object exists.
-	 * @return					If true a valid UDataLayer Object was found and set to OutDataLayer; if false, a valid UDataLayer object was not found and invalid set to OutDataLayer
-	 */
-	bool TryGetDataLayerFromLabel(const FName& DataLayerLabel, UDataLayer*& OutDataLayer);
+	UDataLayerInstance* GetDataLayerInstance(const FName& DataLayerInstanceName) const;
 
 	/**
 	 * Gets all known DataLayers and appends them to the provided array
 	 *
 	 * @param OutDataLayers[OUT] Output array to store all known DataLayers
 	 */
-	void AddAllDataLayersTo(TArray<TWeakObjectPtr<UDataLayer>>& OutDataLayers) const;
+	void AddAllDataLayersTo(TArray<TWeakObjectPtr<UDataLayerInstance>>& OutDataLayers) const;
 
 	/**
-	 * Creates a UDataLayer Object
+	 * Creates a UDataLayerInstance Object
 	 *
-	 * @return	The newly created UDataLayer Object
+	 * @param	ParentDataLayer	An optional data layer that will parent the newly created data layer
+	 * @return	The newly created UDataLayerInstance Object
 	 */
 	UFUNCTION(BlueprintCallable, Category = DataLayers)
-	UDataLayer* CreateDataLayer();
+	UDataLayerInstance* CreateDataLayer(UDataLayerInstance* ParentDataLayer = nullptr);
 
 	/**
 	 * Sets a Parent DataLayer for a specified DataLayer
@@ -517,7 +500,7 @@ public:
 	 *  @return	true if succeeded, false if failed.
 	 */
 	UFUNCTION(BlueprintCallable, Category = DataLayers)
-	bool SetParentDataLayer(UDataLayer* DataLayer, UDataLayer* ParentDataLayer);
+	bool SetParentDataLayer(UDataLayerInstance* DataLayer, UDataLayerInstance* ParentDataLayer);
 
 	/**
 	 * Deletes all of the provided DataLayers
@@ -525,7 +508,7 @@ public:
 	 * @param DataLayersToDelete	A valid list of DataLayer.
 	 */
 	UFUNCTION(BlueprintCallable, Category = DataLayers)
-	void DeleteDataLayers(const TArray<UDataLayer*>& DataLayersToDelete);
+	void DeleteDataLayers(const TArray<UDataLayerInstance*>& DataLayersToDelete);
 
 	/**
 	 * Deletes the provided DataLayer
@@ -533,16 +516,7 @@ public:
 	 * @param DataLayerToDelete		A valid DataLayer
 	 */
 	UFUNCTION(BlueprintCallable, Category = DataLayers)
-	void DeleteDataLayer(UDataLayer* DataLayerToDelete);
-
-	/**
-	 * Renames the provided DataLayer to the provided new label
-	 *
-	 * @param	DataLayer			The DataLayer to be renamed
-	 * @param	NewDataLayerLabel	The new label for the DataLayer to be renamed
-	 */
-	UFUNCTION(BlueprintCallable, Category = DataLayers)
-	bool RenameDataLayer(UDataLayer* DataLayer, const FName& NewDataLayerLabel);
+	void DeleteDataLayer(UDataLayerInstance* DataLayerToDelete);
 
 	/**
 	* Resets user override settings of all DataLayers
@@ -552,16 +526,23 @@ public:
 	/*
 	 * Returns whether the DataLayer contains one of more actors that are part of the editor selection.
 	 */
-	bool DoesDataLayerContainSelectedActors(const UDataLayer* DataLayer) const { return SelectedDataLayersFromEditorSelection.Contains(DataLayer); }
+	bool DoesDataLayerContainSelectedActors(const UDataLayerInstance* DataLayer) const { return SelectedDataLayersFromEditorSelection.Contains(DataLayer); }
+
+	/**
+	* Whether there are any deprecated DataLayerInstance found
+	*/
+	bool HasDeprecatedDataLayers() const;
 
 	//~ Begin Deprecated
 
+	PRAGMA_DISABLE_DEPRECATION_WARNINGS
+
 	UE_DEPRECATED(5.0, "Per-view Data Layer visibility was removed.")
 	UFUNCTION(BlueprintCallable, Category = DataLayers, meta = (DeprecatedFunction, DeprecationMessage = "Per-view Data Layer visibility was removed."))
-	void UpdateAllViewVisibility(UDataLayer* DataLayerThatChanged) {}
+	void UpdateAllViewVisibility(UDEPRECATED_DataLayer* DataLayerThatChanged) {}
 
 	UE_DEPRECATED(5.0, "Per-view Data Layer visibility was removed.")
-	void UpdatePerViewVisibility(FLevelEditorViewportClient* ViewportClient, UDataLayer* DataLayerThatChanged = nullptr) {}
+	void UpdatePerViewVisibility(FLevelEditorViewportClient* ViewportClient, UDEPRECATED_DataLayer* DataLayerThatChanged = nullptr) {}
 
 	UE_DEPRECATED(5.0, "Per-view Data Layer visibility was removed.")
 	void UpdateActorViewVisibility(FLevelEditorViewportClient* ViewportClient, AActor* Actor, const bool bReregisterIfDirty = true) {}
@@ -575,19 +556,36 @@ public:
 
 	UE_DEPRECATED(5.0, "Use SetDataLayerIsLoadedInEditor() instead.")
 	UFUNCTION(BlueprintCallable, Category = DataLayers, meta = (DeprecatedFunction, DeprecationMessage = "Use SetDataLayerIsLoadedInEditor instead"))
-	bool SetDataLayerIsDynamicallyLoadedInEditor(UDataLayer* DataLayer, const bool bIsLoadedInEditor, const bool bIsFromUserChange) { return SetDataLayerIsLoadedInEditor(DataLayer, bIsLoadedInEditor, bIsFromUserChange); }
+	bool SetDataLayerIsDynamicallyLoadedInEditor(UDEPRECATED_DataLayer* DataLayer, const bool bIsLoadedInEditor, const bool bIsFromUserChange) { return false; }
 
 	UE_DEPRECATED(5.0, "Use SetDataLayersIsLoadedInEditor() instead.")
 	UFUNCTION(BlueprintCallable, Category = DataLayers, meta = (DeprecatedFunction, DeprecationMessage = "Use SetDataLayersIsLoadedInEditor instead"))
-	bool SetDataLayersIsDynamicallyLoadedInEditor(const TArray<UDataLayer*>& DataLayers, const bool bIsLoadedInEditor, const bool bIsFromUserChange) { return SetDataLayersIsLoadedInEditor(DataLayers, bIsLoadedInEditor, bIsFromUserChange);  }
+	bool SetDataLayersIsDynamicallyLoadedInEditor(const TArray<UDEPRECATED_DataLayer*>& DataLayers, const bool bIsLoadedInEditor, const bool bIsFromUserChange) { return false;  }
 
 	UE_DEPRECATED(5.0, "Use ToggleDataLayerIsLoadedInEditor() instead.")
 	UFUNCTION(BlueprintCallable, Category = DataLayers, meta = (DeprecatedFunction, DeprecationMessage = "Use ToggleDataLayerIsLoadedInEditor instead"))
-	bool ToggleDataLayerIsDynamicallyLoadedInEditor(UDataLayer* DataLayer, const bool bIsFromUserChange) { return ToggleDataLayerIsLoadedInEditor(DataLayer, bIsFromUserChange); }
+	bool ToggleDataLayerIsDynamicallyLoadedInEditor(UDEPRECATED_DataLayer* DataLayer, const bool bIsFromUserChange) { return false; }
 
 	UE_DEPRECATED(5.0, "Use ToggleDataLayersIsLoadedInEditor() instead.")
 	UFUNCTION(BlueprintCallable, Category = DataLayers, meta = (DeprecatedFunction, DeprecationMessage = "Use ToggleDataLayersIsLoadedInEditor instead"))
-	bool ToggleDataLayersIsDynamicallyLoadedInEditor(const TArray<UDataLayer*>& DataLayers, const bool bIsFromUserChange) { return ToggleDataLayersIsLoadedInEditor(DataLayers, bIsFromUserChange); }
+	bool ToggleDataLayersIsDynamicallyLoadedInEditor(const TArray<UDEPRECATED_DataLayer*>& DataLayers, const bool bIsFromUserChange) { return false; }
+
+	PRAGMA_ENABLE_DEPRECATION_WARNINGS
+
+	UE_DEPRECATED(5.1, "Use GetDataLayerInstance instead")
+	bool TryGetDataLayerFromLabel(const FName& DataLayerLabel, UDataLayerInstance*& OutDataLayer) const;
+
+	UE_DEPRECATED(5.1, "Use GetDataLayerInstance instead")
+	UFUNCTION(BlueprintCallable, Category = DataLayers)
+	UDataLayerInstance* GetDataLayer(const FActorDataLayer& ActorDataLayer) const;
+
+	UE_DEPRECATED(5.1, "Use GetDataLayerInstance instead")
+	UFUNCTION(BlueprintCallable, Category = DataLayers)
+	UDataLayerInstance* GetDataLayerFromLabel(const FName& DataLayerLabel) const;
+
+	UE_DEPRECATED(5.1, "Renaming is not permitted anymore")
+	UFUNCTION(BlueprintCallable, Category = DataLayers)
+	bool RenameDataLayer(UDataLayerInstance* DataLayer, const FName& NewDataLayerLabel);
 
 	//~ End Deprecated
 
@@ -638,17 +636,19 @@ private:
 	/** Delegate handler for FEditorDelegates::PostUndoRedo. It internally calls DataLayerChanged.Broadcast and UpdateAllActorsVisibility to refresh the actors of each DataLayer. */
 	void PostUndoRedo();
 
-	bool SetDataLayerIsLoadedInEditorInternal(UDataLayer* DataLayer, const bool bIsLoadedInEditor, const bool bIsFromUserChange);
+	bool SetDataLayerIsLoadedInEditorInternal(UDataLayerInstance* DataLayer, const bool bIsLoadedInEditor, const bool bIsFromUserChange);
 	bool RefreshWorldPartitionEditorCells(bool bIsFromUserChange);
 	void UpdateDataLayerEditorPerProjectUserSettings();
 
 	void BroadcastActorDataLayersChanged(const TWeakObjectPtr<AActor>& ChangedActor);
-	void BroadcastDataLayerChanged(const EDataLayerAction Action, const TWeakObjectPtr<const UDataLayer>& ChangedDataLayer, const FName& ChangedProperty);
+	void BroadcastDataLayerChanged(const EDataLayerAction Action, const TWeakObjectPtr<const UDataLayerInstance>& ChangedDataLayer, const FName& ChangedProperty);
 	void OnSelectionChanged();
 	void RebuildSelectedDataLayersFromEditorSelection();
 
+	bool PromptDataLayerAssetSelection(FAssetData& OutAsset) const;
+
 	/** Contains Data Layers that contain actors that are part of the editor selection */
-	TSet<TWeakObjectPtr<const UDataLayer>> SelectedDataLayersFromEditorSelection;
+	TSet<TWeakObjectPtr<const UDataLayerInstance>> SelectedDataLayersFromEditorSelection;
 
 	/** Fires whenever one or more DataLayer changes */
 	FOnDataLayerChanged DataLayerChanged;
diff --git a/Engine/Source/Editor/DataLayerEditor/Public/DataLayer/DataLayerPropertyTypeCustomizationHelper.h b/Engine/Source/Editor/DataLayerEditor/Public/DataLayer/DataLayerPropertyTypeCustomizationHelper.h
--- a/Engine/Source/Editor/DataLayerEditor/Public/DataLayer/DataLayerPropertyTypeCustomizationHelper.h
+++ b/Engine/Source/Editor/DataLayerEditor/Public/DataLayer/DataLayerPropertyTypeCustomizationHelper.h
@@ -6,9 +6,9 @@
 #include "Templates/SharedPointer.h"
 #include "Widgets/SWidget.h"
 
-class UDataLayer;
+class UDataLayerInstance;
 
 struct DATALAYEREDITOR_API FDataLayerPropertyTypeCustomizationHelper
 {
-	static TSharedRef<SWidget> CreateDataLayerMenu(TFunction<void(const UDataLayer* DataLayer)> OnDataLayerSelectedFunction);
+	static TSharedRef<SWidget> CreateDataLayerMenu(TFunction<void(const UDataLayerInstance* DataLayer)> OnDataLayerSelectedFunction);
 };
diff --git a/Engine/Source/Editor/DataLayerEditor/Public/DataLayerEditorModule.h b/Engine/Source/Editor/DataLayerEditor/Public/DataLayerEditorModule.h
--- a/Engine/Source/Editor/DataLayerEditor/Public/DataLayerEditorModule.h
+++ b/Engine/Source/Editor/DataLayerEditor/Public/DataLayerEditorModule.h
@@ -7,7 +7,7 @@
 #include "Framework/MultiBox/MultiBoxExtender.h"
 #include "WorldPartition/DataLayer/IDataLayerEditorModule.h"
 
-class UDataLayer;
+class UDataLayerInstance;
 
 /**
  * The module holding all of the UI related pieces for DataLayers
@@ -34,14 +34,13 @@ public:
 	/*
 	 * Selected DataLayer in DataLayer Browser widget
 	 */
-	virtual void SyncDataLayerBrowserToDataLayer(const UDataLayer* DataLayer);
+	virtual void SyncDataLayerBrowserToDataLayer(const UDataLayerInstance* DataLayer);
 
 	/** Delegates to be called to extend the DataLayers menus */
 	DECLARE_DELEGATE_RetVal_OneParam(TSharedRef<FExtender>, FDataLayersMenuExtender, const TSharedRef<FUICommandList>);
 	virtual TArray<FDataLayersMenuExtender>& GetAllDataLayersMenuExtenders() {return DataLayersMenuExtenders;}
 
 private:
-
 	TWeakPtr<SWidget> DataLayerBrowser;
 
 	/** All extender delegates for the DataLayers menus */
diff --git a/Engine/Source/Editor/DetailCustomizations/Private/ActorDetails.cpp b/Engine/Source/Editor/DetailCustomizations/Private/ActorDetails.cpp
--- a/Engine/Source/Editor/DetailCustomizations/Private/ActorDetails.cpp
+++ b/Engine/Source/Editor/DetailCustomizations/Private/ActorDetails.cpp
@@ -156,7 +156,8 @@ void FActorDetails::CustomizeDetails( IDetailLayoutBuilder& DetailLayout )
 		{
 			DetailLayout.HideProperty(DetailLayout.GetProperty(AActor::GetRuntimeGridPropertyName(), AActor::StaticClass()));
 			DetailLayout.HideProperty(DetailLayout.GetProperty(AActor::GetIsSpatiallyLoadedPropertyName(), AActor::StaticClass()));
-			DetailLayout.HideProperty(DetailLayout.GetProperty(AActor::GetDataLayersPropertyName(), AActor::StaticClass()));
+			DetailLayout.HideProperty(DetailLayout.GetProperty(AActor::GetDataLayerAssetsPropertyName(), AActor::StaticClass()));
+			DetailLayout.HideProperty(DetailLayout.GetProperty(AActor::GetDataLayerPropertyName(), AActor::StaticClass()));
 			DetailLayout.HideProperty(DetailLayout.GetProperty(AActor::GetHLODLayerPropertyName(), AActor::StaticClass()));
 		}
 
diff --git a/Engine/Source/Editor/FoliageEdit/Private/FoliageEdMode.cpp b/Engine/Source/Editor/FoliageEdit/Private/FoliageEdMode.cpp
--- a/Engine/Source/Editor/FoliageEdit/Private/FoliageEdMode.cpp
+++ b/Engine/Source/Editor/FoliageEdit/Private/FoliageEdMode.cpp
@@ -2049,7 +2049,7 @@ void FEdModeFoliage::RemoveSelectedInstances(UWorld* InWorld)
 void FEdModeFoliage::GetFoliageTypeFilters(TArray<const UClass*>& OutFilters) const
 {
 	OutFilters.Add(UFoliageType_InstancedStaticMesh::StaticClass());
-	if (UWorld::IsPartitionedWorld(GetWorld()))
+	if (GetWorld()->HasSubsystem<UWorldPartitionSubsystem>())
 	{
 		return;
 	}
diff --git a/Engine/Source/Editor/FoliageEdit/Private/SFoliageEdit.cpp b/Engine/Source/Editor/FoliageEdit/Private/SFoliageEdit.cpp
--- a/Engine/Source/Editor/FoliageEdit/Private/SFoliageEdit.cpp
+++ b/Engine/Source/Editor/FoliageEdit/Private/SFoliageEdit.cpp
@@ -34,7 +34,7 @@
 #include "Editor/PropertyEditor/Public/VariablePrecisionNumericInterface.h"
 #include "DataLayer/DataLayerEditorSubsystem.h"
 #include "DataLayer/DataLayerPropertyTypeCustomizationHelper.h"
-#include "WorldPartition/DataLayer/DataLayer.h"
+#include "WorldPartition/DataLayer/DataLayerInstance.h"
 #include "WorldPartition/WorldPartitionSubsystem.h"
 
 #define LOCTEXT_NAMESPACE "FoliageEd_Mode"
diff --git a/Engine/Source/Editor/LevelEditor/Private/SLevelViewport.cpp b/Engine/Source/Editor/LevelEditor/Private/SLevelViewport.cpp
--- a/Engine/Source/Editor/LevelEditor/Private/SLevelViewport.cpp
+++ b/Engine/Source/Editor/LevelEditor/Private/SLevelViewport.cpp
@@ -86,7 +86,7 @@
 #include "WorldPartition/WorldPartitionSubsystem.h"
 #include "WorldPartition/WorldPartition.h"
 #include "WorldPartition/DataLayer/WorldDataLayers.h"
-#include "WorldPartition/DataLayer/DataLayer.h"
+#include "WorldPartition/DataLayer/DataLayerInstance.h"
 #include "DataLayer/DataLayerEditorSubsystem.h"
 #include "SInViewportDetails.h"
 #include "Viewports/InViewportUIDragOperation.h"
diff --git a/Engine/Source/Editor/MovieSceneTools/Private/TrackEditors/DataLayerTrackEditor.cpp b/Engine/Source/Editor/MovieSceneTools/Private/TrackEditors/DataLayerTrackEditor.cpp
--- a/Engine/Source/Editor/MovieSceneTools/Private/TrackEditors/DataLayerTrackEditor.cpp
+++ b/Engine/Source/Editor/MovieSceneTools/Private/TrackEditors/DataLayerTrackEditor.cpp
@@ -118,20 +118,20 @@ struct FDataLayerSection
 		{
 			FString LayerName;
 
-			const TArray<FActorDataLayer>& DataLayers = Section->GetDataLayers();
-			for (int32 Index = 0; Index < DataLayers.Num(); ++Index)
+			const TArray<UDataLayerAsset*>& DataLayerAssets = Section->GetDataLayerAssets();
+			for (int32 Index = 0; Index < DataLayerAssets.Num(); ++Index)
 			{
-				UDataLayer* DataLayer = SubSystem->GetDataLayer(DataLayers[Index]);
-				if (DataLayer)
+				UDataLayerInstance* DataLayerInstance = SubSystem->GetDataLayerInstance(DataLayerAssets[Index]);
+				if (DataLayerInstance)
 				{
-					LayerName += DataLayer->GetDataLayerLabel().ToString();
+					LayerName += DataLayerInstance->GetDataLayerFullName();
 				}
 				else
 				{
-					LayerName += FText::Format(LOCTEXT("UnknownDataLayer", "**invalid: {0}**"), FText::FromString(DataLayers[Index].Name.ToString())).ToString();
+					LayerName += FText::Format(LOCTEXT("UnknownDataLayer", "**invalid: {0}**"), FText::FromString(DataLayerAssets[Index]->GetFullName())).ToString();
 				}
 
-				if (Index < DataLayers.Num()-1)
+				if (Index < DataLayerAssets.Num()-1)
 				{
 					LayerName += TEXT(", ");
 				}
diff --git a/Engine/Source/Editor/SceneOutliner/Private/SceneOutlinerModule.cpp b/Engine/Source/Editor/SceneOutliner/Private/SceneOutlinerModule.cpp
--- a/Engine/Source/Editor/SceneOutliner/Private/SceneOutlinerModule.cpp
+++ b/Engine/Source/Editor/SceneOutliner/Private/SceneOutlinerModule.cpp
@@ -22,7 +22,7 @@
 #include "ActorDescTreeItem.h"
 #include "FolderTreeItem.h"
 #include "WorldTreeItem.h"
-#include "WorldPartition/DataLayer/DataLayer.h"
+#include "WorldPartition/DataLayer/DataLayerInstance.h"
 #include "WorldPartition/DataLayer/WorldDataLayers.h"
 
 #define LOCTEXT_NAMESPACE "SceneOutlinerModule"
@@ -232,7 +232,7 @@ void FSceneOutlinerModule::CreateActorInfoColumns(FSceneOutlinerInitializationOp
 	{
 		TStringBuilder<128> Builder;
 
-		TArray<const UDataLayer*> DataLayerObjects;
+		TArray<const UDataLayerInstance*> DataLayerInstances;
 
 		if (const FActorTreeItem* ActorItem = Item.CastTo<FActorTreeItem>())
 		{
@@ -240,12 +240,12 @@ void FSceneOutlinerModule::CreateActorInfoColumns(FSceneOutlinerInitializationOp
 
 			if (Actor)
 			{
-				DataLayerObjects = Actor->GetDataLayerObjects();
+				DataLayerInstances = Actor->GetDataLayerInstances();
 			}
 		}
 		else if (const FActorDescTreeItem* ActorDescItem = Item.CastTo<FActorDescTreeItem>())
 		{
-			if (const FWorldPartitionActorDesc* ActorDesc = ActorDescItem->ActorDescHandle.Get(); ActorDesc && !ActorDesc->GetDataLayers().IsEmpty())
+			if (const FWorldPartitionActorDesc* ActorDesc = ActorDescItem->ActorDescHandle.Get(); ActorDesc && !ActorDesc->GetDataLayerInstanceNames().IsEmpty())
 			{
 				const UActorDescContainer* ActorDescContainer = ActorDescItem->ActorDescHandle.Container.Get();
 				const UWorld* World = ActorDescContainer ? ActorDescContainer->GetWorld() : nullptr;
@@ -253,18 +253,18 @@ void FSceneOutlinerModule::CreateActorInfoColumns(FSceneOutlinerInitializationOp
 			
 				if (WorldDataLayers)
 				{
-					DataLayerObjects = WorldDataLayers->GetDataLayerObjects(ActorDesc->GetDataLayers());
+					DataLayerInstances = WorldDataLayers->GetDataLayerInstances(ActorDesc->GetDataLayerInstanceNames());
 				}
 			}
 		}
 
-		for (const UDataLayer* DataLayer : DataLayerObjects)
+		for (const UDataLayerInstance* DataLayerInstance : DataLayerInstances)
 		{
 			if (Builder.Len())
 			{
 				Builder += TEXT(", ");
 			}
-			Builder += DataLayer->GetDataLayerLabel().ToString();
+			Builder += DataLayerInstance->GetDataLayerShortName();
 		}
 		
 		return Builder.ToString();
diff --git a/Engine/Source/Editor/UnrealEd/Classes/Commandlets/WorldPartitionConvertCommandlet.h b/Engine/Source/Editor/UnrealEd/Classes/Commandlets/WorldPartitionConvertCommandlet.h
--- a/Engine/Source/Editor/UnrealEd/Classes/Commandlets/WorldPartitionConvertCommandlet.h
+++ b/Engine/Source/Editor/UnrealEd/Classes/Commandlets/WorldPartitionConvertCommandlet.h
@@ -17,6 +17,7 @@ UNREALED_API DECLARE_LOG_CATEGORY_EXTERN(LogWorldPartitionConvertCommandlet, Log
 class ULevelStreaming;
 class UWorldPartition;
 class ULevelStreaming;
+class UDataLayerFactory;
 
 USTRUCT()
 struct FHLODLayerActorMapping
@@ -133,4 +134,10 @@ protected:
 
 	UPROPERTY(Config)
 	uint32 LandscapeGridSize;
+
+	UPROPERTY(Config)
+	FString DataLayerAssetFolder;
+
+	UPROPERTY()
+	TObjectPtr<UDataLayerFactory> DataLayerFactory;
 };
diff --git a/Engine/Source/Editor/UnrealEd/Private/Commandlets/WorldPartitionConvertCommandlet.cpp b/Engine/Source/Editor/UnrealEd/Private/Commandlets/WorldPartitionConvertCommandlet.cpp
--- a/Engine/Source/Editor/UnrealEd/Private/Commandlets/WorldPartitionConvertCommandlet.cpp
+++ b/Engine/Source/Editor/UnrealEd/Private/Commandlets/WorldPartitionConvertCommandlet.cpp
@@ -6,6 +6,7 @@
 
 #include "Commandlets/WorldPartitionConvertCommandlet.h"
 #include "Algo/ForEach.h"
+#include "AssetToolsModule.h"
 #include "HAL/PlatformFileManager.h"
 #include "Misc/Paths.h"
 #include "Misc/ConfigCacheIni.h"
@@ -23,6 +24,7 @@
 #include "WorldPartition/WorldPartitionMiniMap.h"
 #include "WorldPartition/WorldPartitionMiniMapHelper.h"
 #include "LevelInstance/LevelInstanceActor.h"
+#include "DataLayer/DataLayerFactory.h"
 #include "GameFramework/WorldSettings.h"
 #include "UObject/UObjectHash.h"
 #include "PackageHelperFunctions.h"
@@ -53,7 +55,7 @@
 #include "LandscapeSplinesComponent.h"
 #include "LandscapeSplineControlPoint.h"
 #include "LandscapeGizmoActor.h"
-#include "WorldPartition/DataLayer/DataLayer.h"
+#include "WorldPartition/DataLayer/DataLayerInstanceWithAsset.h"
 #include "WorldPartition/DataLayer/WorldDataLayers.h"
 #include "ActorFolder.h"
 
@@ -163,6 +165,8 @@ UWorldPartitionConvertCommandlet::UWorldPartitionConvertCommandlet(const FObject
 	, WorldOrigin(FVector::ZeroVector)
 	, WorldExtent(WORLDPARTITION_MAX * 0.5)
 	, LandscapeGridSize(4)
+	, DataLayerAssetFolder(TEXT("Game/DataLayers/"))
+	, DataLayerFactory(NewObject<UDataLayerFactory>())
 {}
 
 UWorld* UWorldPartitionConvertCommandlet::LoadWorld(const FString& LevelToLoad)
@@ -733,6 +737,11 @@ int32 UWorldPartitionConvertCommandlet::Main(const FString& Params)
 		return 1;
 	}
 
+	if (FString* DataLayerAssetFolderValue = Arguments.Find(TEXT("DataLayerAssetFolder")))
+	{
+		DataLayerAssetFolder = *DataLayerAssetFolderValue;
+	}
+
 	ReadAdditionalTokensAndSwitches(Tokens, Switches);
 
 	if (bVerbose)
@@ -1003,7 +1012,7 @@ int32 UWorldPartitionConvertCommandlet::Main(const FString& Params)
 		}
 	};
 
-	auto PrepareLevelActors = [this, PartitionFoliage, PartitionLandscape, MainWorldDataLayers](ULevel* Level, TArray<AActor*>& Actors, bool bMainLevel) -> bool
+	auto PrepareLevelActors = [this, PartitionFoliage, PartitionLandscape, MainWorldDataLayers, MainWorld](ULevel* Level, TArray<AActor*>& Actors, bool bMainLevel) -> bool
 	{
 		TRACE_CPUPROFILER_EVENT_SCOPE(PrepareLevelActors);
 
@@ -1011,6 +1020,8 @@ int32 UWorldPartitionConvertCommandlet::Main(const FString& Params)
 
 		TArray<AInstancedFoliageActor*> IFAs;
 		TSet<ULandscapeInfo*> LandscapeInfos;
+		IAssetTools& AssetTools = FModuleManager::GetModuleChecked<FAssetToolsModule>("AssetTools").Get();
+		FString DataLayerFolder = DataLayerAssetFolder + MainWorld->GetName() + TEXT("/");
 		for (AActor* Actor: Actors)
 		{
 			if (Actor && IsValidChecked(Actor))
@@ -1049,14 +1060,18 @@ int32 UWorldPartitionConvertCommandlet::Main(const FString& Params)
 					{
 						for (FName Layer : Actor->Layers)
 						{
-							UDataLayer* DataLayer = const_cast<UDataLayer*>(MainWorldDataLayers->GetDataLayerFromLabel(Layer));
-							if (!DataLayer)
+							FName DataLayerAssetFullName(DataLayerAssetFolder + Layer.ToString());
+							UDataLayerInstance* DataLayerInstance = const_cast<UDataLayerInstance*>(MainWorldDataLayers->GetDataLayerInstanceFromAssetName(DataLayerAssetFullName));
+							if (!DataLayerInstance)
 							{
-								DataLayer = MainWorldDataLayers->CreateDataLayer();
-								DataLayer->SetDataLayerLabel(Layer);
-								DataLayer->SetIsRuntime(false);
+								if (UObject* Asset = AssetTools.CreateAsset(Layer.ToString(), DataLayerFolder, UDataLayerAsset::StaticClass(), DataLayerFactory))
+								{
+									UDataLayerAsset* DataLayerAsset = CastChecked<UDataLayerAsset>(Asset);
+									DataLayerAsset->SetType(EDataLayerType::Editor);
+									DataLayerInstance = MainWorldDataLayers->CreateDataLayer<UDataLayerInstanceWithAsset>(DataLayerAsset);
+								}
 							}
-							Actor->AddDataLayer(DataLayer);
+							Actor->AddDataLayer(DataLayerInstance);
 						}
 					}
 					// Clear actor layers as they are not supported yet in world partition, keep them if only merging
diff --git a/Engine/Source/Editor/UnrealEd/Private/WorldPartition/WorldPartitionBuilder.cpp b/Engine/Source/Editor/UnrealEd/Private/WorldPartition/WorldPartitionBuilder.cpp
--- a/Engine/Source/Editor/UnrealEd/Private/WorldPartition/WorldPartitionBuilder.cpp
+++ b/Engine/Source/Editor/UnrealEd/Private/WorldPartition/WorldPartitionBuilder.cpp
@@ -13,7 +13,7 @@
 #include "WorldPartition/WorldPartitionHelpers.h"
 #include "WorldPartition/WorldPartitionSubsystem.h"
 #include "WorldPartition/DataLayer/WorldDataLayers.h"
-#include "WorldPartition/DataLayer/DataLayer.h"
+#include "WorldPartition/DataLayer/DataLayerInstance.h"
 #include "LevelInstance/LevelInstanceSubsystem.h"
 #include "Math/IntVector.h"
 #include "UObject/SavePackage.h"
@@ -162,14 +162,14 @@ bool UWorldPartitionBuilder::Run(UWorld* World, FPackageSourceControlHelper& Pac
 	{
 		// Load Data Layers
 		bool bUpdateEditorCells = false;
-		WorldDataLayers->ForEachDataLayer([&bUpdateEditorCells, this](UDataLayer* DataLayer)
+		WorldDataLayers->ForEachDataLayer([&bUpdateEditorCells, this](UDataLayerInstance* DataLayer)
 		{
-			const FName DataLayerLabel = DataLayer->GetDataLayerLabel();
+			const FName DataLayerShortName(DataLayer->GetDataLayerShortName());
 
 			// Load all Non Excluded Data Layers + Non DynamicallyLoaded Data Layers + Initially Active Data Layers + Data Layers provided by builder
-			const bool bLoadedInEditor = !ExcludedDataLayerLabels.Contains(DataLayerLabel) && ((bLoadNonDynamicDataLayers && !DataLayer->IsRuntime()) ||
+			const bool bLoadedInEditor = !ExcludedDataLayerShortNames.Contains(DataLayerShortName) && ((bLoadNonDynamicDataLayers && !DataLayer->IsRuntime()) ||
 										 (bLoadInitiallyActiveDataLayers && DataLayer->GetInitialRuntimeState() == EDataLayerRuntimeState::Activated) ||
-										 DataLayerLabels.Contains(DataLayerLabel));
+										DataLayerShortNames.Contains(DataLayerShortName));
 			if (DataLayer->IsLoadedInEditor() != bLoadedInEditor)
 			{
 				bUpdateEditorCells = true;
@@ -180,11 +180,11 @@ bool UWorldPartitionBuilder::Run(UWorld* World, FPackageSourceControlHelper& Pac
 				}
 			}
 			
-			UE_LOG(LogWorldPartitionBuilder, Display, TEXT("DataLayer '%s' Loaded: %d"), *UDataLayer::GetDataLayerText(DataLayer).ToString(), bLoadedInEditor ? 1 : 0);
+			UE_LOG(LogWorldPartitionBuilder, Display, TEXT("DataLayer '%s' Loaded: %d"), *UDataLayerInstance::GetDataLayerText(DataLayer).ToString(), bLoadedInEditor ? 1 : 0);
 			
 			return true;
 		});
-
+	
 		if (bUpdateEditorCells)
 		{
 			UE_LOG(LogWorldPartitionBuilder, Display, TEXT("DataLayer load state changed refreshing editor cells"));
diff --git a/Engine/Source/Editor/UnrealEd/Private/WorldPartition/WorldPartitionHLODsBuilder.cpp b/Engine/Source/Editor/UnrealEd/Private/WorldPartition/WorldPartitionHLODsBuilder.cpp
--- a/Engine/Source/Editor/UnrealEd/Private/WorldPartition/WorldPartitionHLODsBuilder.cpp
+++ b/Engine/Source/Editor/UnrealEd/Private/WorldPartition/WorldPartitionHLODsBuilder.cpp
@@ -21,6 +21,7 @@
 #include "AssetRegistryModule.h"
 #include "DerivedDataCacheInterface.h"
 
+#include "WorldPartition/DataLayer/DataLayerInstance.h"
 #include "WorldPartition/WorldPartition.h"
 #include "WorldPartition/WorldPartitionHelpers.h"
 #include "WorldPartition/HLOD/HLODActor.h"
@@ -669,7 +670,7 @@ bool UWorldPartitionHLODsBuilder::DumpStats()
 		HLODStat.RuntimeGrid = HLODActor->GetRuntimeGrid();
 		HLODStat.bIsSpatiallyLoaded = HLODActor->GetIsSpatiallyLoaded();
 		HLODActor->GetGridIndices(HLODStat.GridLocationX, HLODStat.GridLocationY, HLODStat.GridLocationZ);
-		HLODStat.DataLayers = *FString::JoinBy(HLODActor->GetDataLayerObjects(), TEXT("| "), [](const UDataLayer* DataLayer) { return DataLayer->GetDataLayerLabel().ToString(); });
+		HLODStat.DataLayers = *FString::JoinBy(HLODActor->GetDataLayerInstances(), TEXT("| "), [](const UDataLayerInstance* DataLayer) { return DataLayer->GetDataLayerShortName(); });
 		HLODStat.HLODType = HLODActor->GetSubActorsHLODLayer()->GetLayerType();
 
 		HLODActor->ForEachComponent<UInstancedStaticMeshComponent>(false, [&](const UInstancedStaticMeshComponent* ISMC)
diff --git a/Engine/Source/Editor/UnrealEd/Private/WorldPartition/WorldPartitionMiniMapBuilder.cpp b/Engine/Source/Editor/UnrealEd/Private/WorldPartition/WorldPartitionMiniMapBuilder.cpp
--- a/Engine/Source/Editor/UnrealEd/Private/WorldPartition/WorldPartitionMiniMapBuilder.cpp
+++ b/Engine/Source/Editor/UnrealEd/Private/WorldPartition/WorldPartitionMiniMapBuilder.cpp
@@ -20,7 +20,7 @@
 
 #include "WorldPartition/ActorDescContainer.h"
 #include "WorldPartition/DataLayer/ActorDataLayer.h"
-#include "WorldPartition/DataLayer/DataLayer.h"
+#include "WorldPartition/DataLayer/DataLayerInstance.h"
 #include "WorldPartition/DataLayer/WorldDataLayers.h"
 #include "WorldPartition/HLOD/HLODActor.h"
 #include "WorldPartition/HLOD/HLODActorDesc.h"
@@ -65,11 +65,11 @@ bool UWorldPartitionMiniMapBuilder::PreRun(UWorld* World, FPackageSourceControlH
 
 	for (const FActorDataLayer& ActorDataLayer : WorldMiniMap->ExcludedDataLayers)
 	{
-		const UDataLayer* DataLayer = WorldDataLayers->GetDataLayerFromName(ActorDataLayer.Name);
+		const UDataLayerInstance* DataLayerInstance = WorldDataLayers->GetDataLayerInstance(ActorDataLayer.Name);
 
-		if (DataLayer != nullptr)
+		if (DataLayerInstance != nullptr)
 		{
-			ExcludedDataLayerLabels.Add(DataLayer->GetDataLayerLabel());
+			ExcludedDataLayerShortNames.Add(FName(DataLayerInstance->GetDataLayerShortName()));
 		}
 	}
 
diff --git a/Engine/Source/Editor/UnrealEd/Private/WorldPartition/WorldPartitionNavigationDataBuilder.cpp b/Engine/Source/Editor/UnrealEd/Private/WorldPartition/WorldPartitionNavigationDataBuilder.cpp
--- a/Engine/Source/Editor/UnrealEd/Private/WorldPartition/WorldPartitionNavigationDataBuilder.cpp
+++ b/Engine/Source/Editor/UnrealEd/Private/WorldPartition/WorldPartitionNavigationDataBuilder.cpp
@@ -29,11 +29,11 @@ bool UWorldPartitionNavigationDataBuilder::PreRun(UWorld* World, FPackageSourceC
 	// Set runtime data layer to be included in the generation.
 	if (const AWorldDataLayers* WorldDataLayers = World->GetWorldDataLayers())
 	{
-		WorldDataLayers->ForEachDataLayer([this](const UDataLayer* DataLayer)
+		WorldDataLayers->ForEachDataLayer([this](const UDataLayerInstance* DataLayer)
 		{
 			if (DataLayer->IsRuntime())
 			{
-				DataLayerLabels.Add(DataLayer->GetDataLayerLabel());
+				DataLayerShortNames.Add(FName(DataLayer->GetDataLayerShortName()));
 			}
 
 			return true;
diff --git a/Engine/Source/Editor/UnrealEd/Public/WorldPartition/WorldPartitionBuilder.h b/Engine/Source/Editor/UnrealEd/Public/WorldPartition/WorldPartitionBuilder.h
--- a/Engine/Source/Editor/UnrealEd/Public/WorldPartition/WorldPartitionBuilder.h
+++ b/Engine/Source/Editor/UnrealEd/Public/WorldPartition/WorldPartitionBuilder.h
@@ -83,8 +83,8 @@ protected:
 
 	int32 IterativeCellSize = 102400;
 	int32 IterativeCellOverlapSize = 0;
-	TSet<FName> DataLayerLabels;
-	TSet<FName> ExcludedDataLayerLabels;
+	TSet<FName> DataLayerShortNames;
+	TSet<FName> ExcludedDataLayerShortNames;
 	bool bLoadNonDynamicDataLayers = true;
 	bool bLoadInitiallyActiveDataLayers = true;
 
diff --git a/Engine/Source/Runtime/Engine/Classes/GameFramework/Actor.h b/Engine/Source/Runtime/Engine/Classes/GameFramework/Actor.h
--- a/Engine/Source/Runtime/Engine/Classes/GameFramework/Actor.h
+++ b/Engine/Source/Runtime/Engine/Classes/GameFramework/Actor.h
@@ -39,7 +39,9 @@ class UPrimitiveComponent;
 struct FAttachedActorInfo;
 struct FNetViewer;
 struct FNetworkObjectInfo;
-class UDataLayer;
+class UDEPRECATED_DataLayer;
+class UDataLayerAsset;
+class UDataLayerInstance;
 class AWorldDataLayers;
 class UActorFolder;
 
@@ -850,10 +852,13 @@ protected:
 	FGuid ActorGuid;
 
 	/** DataLayers the actor belongs to.*/
-	UPROPERTY(EditAnywhere, AdvancedDisplay, Category = DataLayers)
+	UPROPERTY(VisibleAnywhere, AdvancedDisplay, Category = DataLayers)
 	TArray<FActorDataLayer> DataLayers;
 
-	TArray<FActorDataLayer> PreEditChangeDataLayers;
+	UPROPERTY(EditAnywhere, AdvancedDisplay, Category = DataLayers)
+	TArray<TObjectPtr<const UDataLayerAsset>> DataLayerAssets;
+
+	TArray<TObjectPtr<const UDataLayerAsset>> PreEditChangeDataLayers;
 
 public:
 	/** The copy/paste id used to remap actors during copy operations */
@@ -1167,22 +1172,52 @@ public:
 #if WITH_EDITOR
 	//~=============================================================================
 	// DataLayers functions.
-	bool AddDataLayer(const UDataLayer* DataLayer);
-	bool RemoveDataLayer(const UDataLayer* DataLayer);
+	bool AddDataLayer(const UDataLayerAsset* DataLayerAsset);
+	bool AddDataLayer(const UDataLayerInstance* DataLayerInstance);
+	bool RemoveDataLayer(const UDataLayerAsset* DataLayerAsset);
+	bool RemoveDataLayer(const UDataLayerInstance* DataLayerInstance);
 	bool RemoveAllDataLayers();
-	bool ContainsDataLayer(const UDataLayer* DataLayer) const;
+	bool ContainsDataLayer(const UDataLayerAsset* DataLayerAsset) const;
+	bool ContainsDataLayer(const UDataLayerInstance* DataLayerInstance) const;
 	virtual bool SupportsDataLayer() const { return true; }
 	bool HasDataLayers() const;
 	bool HasValidDataLayers() const;
-	bool HasAllDataLayers(const TArray<const UDataLayer*>& DataLayers) const;
-	bool HasAnyOfDataLayers(const TArray<FName>& DataLayerNames) const;
-	TArray<FName> GetDataLayerNames() const;
-	TArray<const UDataLayer*> GetDataLayerObjects() const;
-	TArray<const UDataLayer*> GetDataLayerObjects(const AWorldDataLayers* WorldDataLayers) const;
+	TArray<FName> GetDataLayerInstanceNames() const;
+	TArray<const UDataLayerInstance*> GetDataLayerInstances() const;
+	TArray<const UDataLayerInstance*> GetDataLayerInstances(const AWorldDataLayers* WorldDataLayers) const;
 	bool IsPropertyChangedAffectingDataLayers(FPropertyChangedEvent& PropertyChangedEvent) const;
 	bool IsValidForDataLayer() const;
 	void FixupDataLayers(bool bRevertChangesOnLockedDataLayer = false);
-	static const FName GetDataLayersPropertyName() { return GET_MEMBER_NAME_CHECKED(AActor, DataLayers); }
+	static const FName GetDataLayerAssetsPropertyName() { return GET_MEMBER_NAME_CHECKED(AActor, DataLayerAssets); }
+	static const FName GetDataLayerPropertyName() { return GET_MEMBER_NAME_CHECKED(AActor, DataLayers); }
+
+	const TArray<TObjectPtr<const UDataLayerAsset>>& GetDataLayerAssets() const { return DataLayerAssets; }
+
+	//~ Begin Deprecated
+
+	PRAGMA_DISABLE_DEPRECATION_WARNINGS
+
+	UE_DEPRECATED(5.1, "Convert DataLayer using UDataLayerToAssetCommandlet and use AddDataLayer(UDataLayerInstance*)")
+	bool AddDataLayer(const UDEPRECATED_DataLayer* DataLayer);
+
+	UE_DEPRECATED(5.1, "Convert DataLayer using UDataLayerToAssetCommandlet and use RemoveDataLayer(UDataLayerInstance*)")
+	bool RemoveDataLayer(const UDEPRECATED_DataLayer* DataLayer);
+
+	PRAGMA_ENABLE_DEPRECATION_WARNINGS
+
+	UE_DEPRECATED(5.1, "Convert DataLayer using UDataLayerToAssetCommandlet and use AddDataLayer(UDataLayerAsset*)")
+	bool AddDataLayer(const FActorDataLayer& ActorDataLayer);
+
+	UE_DEPRECATED(5.1, "Convert DataLayer using UDataLayerToAssetCommandlet and use RemoveDataLayer(UDataLayerAsset*)")
+	bool RemoveDataLayer(const FActorDataLayer& ActorDataLayer);
+
+	UE_DEPRECATED(5.1, "Convert DataLayer using UDataLayerToAssetCommandlet and use ContainsDataLayer(const UDataLayerAsset*)")
+	bool ContainsDataLayer(const FActorDataLayer& ActorDataLayer) const;
+
+	UE_DEPRECATED(5.1, "Convert DataLayer using UDataLayerToAssetCommandlet and use GetDataLayerAssets() instead")
+	TArray<FActorDataLayer> const& GetActorDataLayers() const { return DataLayers; }
+
+	//~ End Deprecated
 #endif
 
 	//~=============================================================================
diff --git a/Engine/Source/Runtime/Engine/Classes/GameFramework/LightWeightInstanceBlueprintFunctionLibrary.h b/Engine/Source/Runtime/Engine/Classes/GameFramework/LightWeightInstanceBlueprintFunctionLibrary.h
--- a/Engine/Source/Runtime/Engine/Classes/GameFramework/LightWeightInstanceBlueprintFunctionLibrary.h
+++ b/Engine/Source/Runtime/Engine/Classes/GameFramework/LightWeightInstanceBlueprintFunctionLibrary.h
@@ -16,7 +16,7 @@ public:
 
 	// Returns a handle to a new light weight instance that represents an object of type ActorClass
 	UFUNCTION(BlueprintCallable, Category = "Light Weight Instance")
-	static FActorInstanceHandle CreateNewLightWeightInstance(UClass* ActorClass, FTransform Transform, UDataLayer* Layer, UWorld* World);
+	static FActorInstanceHandle CreateNewLightWeightInstance(UClass* ActorClass, FTransform Transform, UDataLayerInstance* Layer, UWorld* World);
 
 	// Returns a handle to the light weight representation and destroys Actor if successful; Returns a handle to Actor otherwise
 	UFUNCTION(BlueprintCallable, Category = "Light Weight Instance")
diff --git a/Engine/Source/Runtime/Engine/Classes/GameFramework/LightWeightInstanceSubsystem.h b/Engine/Source/Runtime/Engine/Classes/GameFramework/LightWeightInstanceSubsystem.h
--- a/Engine/Source/Runtime/Engine/Classes/GameFramework/LightWeightInstanceSubsystem.h
+++ b/Engine/Source/Runtime/Engine/Classes/GameFramework/LightWeightInstanceSubsystem.h
@@ -44,11 +44,11 @@ struct ENGINE_API FLightWeightInstanceSubsystem
 	ALightWeightInstanceManager* FindLightWeightInstanceManager(const FActorInstanceHandle& Handle) const;
 
 	// Returns the instance manager that handles actors of type ActorClass in level Level
-	ALightWeightInstanceManager* FindLightWeightInstanceManager(UClass* ActorClass, const UDataLayer* Layer) const;
+	ALightWeightInstanceManager* FindLightWeightInstanceManager(UClass* ActorClass, const UDataLayerInstance* Layer) const;
 
 	// Returns the instance manager that handles instances of type Class that live in Level
 	UFUNCTION(Server, Unreliable)
-	ALightWeightInstanceManager* FindOrAddLightWeightInstanceManager(UClass* ActorClass, const UDataLayer* Layer, UWorld* World);
+	ALightWeightInstanceManager* FindOrAddLightWeightInstanceManager(UClass* ActorClass, const UDataLayerInstance* Layer, UWorld* World);
 
 	// Returns the actor specified by Handle. This may require loading and creating the actor object.
 	AActor* FetchActor(const FActorInstanceHandle& Handle);
@@ -69,7 +69,7 @@ struct ENGINE_API FLightWeightInstanceSubsystem
 	bool IsInLevel(const FActorInstanceHandle& Handle, const ULevel* InLevel);
 
 	// Returns a handle to a new light weight instance that represents an object of type ActorClass
-	FActorInstanceHandle CreateNewLightWeightInstance(UClass* ActorClass, FLWIData* InitData, UDataLayer* Layer, UWorld* World);
+	FActorInstanceHandle CreateNewLightWeightInstance(UClass* ActorClass, FLWIData* InitData, UDataLayerInstance* Layer, UWorld* World);
 
 	// deletes the instance identified by Handle
 	void DeleteInstance(const FActorInstanceHandle& Handle);
diff --git a/Engine/Source/Runtime/Engine/Private/ActorEditor.cpp b/Engine/Source/Runtime/Engine/Private/ActorEditor.cpp
--- a/Engine/Source/Runtime/Engine/Private/ActorEditor.cpp
+++ b/Engine/Source/Runtime/Engine/Private/ActorEditor.cpp
@@ -14,7 +14,7 @@
 #include "Engine/BlueprintGeneratedClass.h"
 #include "WorldPartition/WorldPartitionSubsystem.h"
 #include "WorldPartition/DataLayer/WorldDataLayers.h"
-#include "WorldPartition/DataLayer/DataLayer.h"
+#include "WorldPartition/DataLayer/DataLayerInstance.h"
 #include "EditorSupportDelegates.h"
 #include "Logging/TokenizedMessage.h"
 #include "Logging/MessageLog.h"
@@ -54,10 +54,9 @@ void AActor::PreEditChange(FProperty* PropertyThatWillChange)
 	}
 
 	PreEditChangeDataLayers.Reset();
-	if (PropertyThatWillChange->GetFName() == GET_MEMBER_NAME_CHECKED(AActor, DataLayers) ||
-		PropertyThatWillChange->GetFName() == GET_MEMBER_NAME_CHECKED(FActorDataLayer, Name))
+	if (PropertyThatWillChange->GetFName() == GET_MEMBER_NAME_CHECKED(AActor, DataLayerAssets))
 	{
-		PreEditChangeDataLayers = DataLayers;
+		PreEditChangeDataLayers = DataLayerAssets;
 	}
 }
 
@@ -71,7 +70,7 @@ bool AActor::CanEditChange(const FProperty* PropertyThatWillChange) const
 
 	const bool bIsSpatiallyLoadedProperty = PropertyThatWillChange->GetFName() == GET_MEMBER_NAME_CHECKED(AActor, bIsSpatiallyLoaded);
 	const bool bIsRuntimeGridProperty = PropertyThatWillChange->GetFName() == GET_MEMBER_NAME_CHECKED(AActor, RuntimeGrid);
-	const bool bIsDataLayersProperty = PropertyThatWillChange->GetFName() == GET_MEMBER_NAME_CHECKED(AActor, DataLayers);
+	const bool bIsDataLayersProperty = PropertyThatWillChange->GetFName() == GET_MEMBER_NAME_CHECKED(AActor, DataLayerAssets);
 	const bool bIsHLODLayerProperty = PropertyThatWillChange->GetFName() == GET_MEMBER_NAME_CHECKED(AActor, HLODLayer);
 
 	if (bIsSpatiallyLoadedProperty || bIsRuntimeGridProperty || bIsDataLayersProperty || bIsHLODLayerProperty)
@@ -1316,26 +1315,32 @@ EDataValidationResult AActor::IsDataValid(TArray<FText>& ValidationErrors)
 //---------------------------------------------------------------------------
 // DataLayers (begin)
 
-bool AActor::AddDataLayer(const UDataLayer* DataLayer)
+bool AActor::AddDataLayer(const UDataLayerInstance* DataLayerInstance)
+{
+	return DataLayerInstance->AddActor(this);
+}
+
+bool AActor::AddDataLayer(const UDataLayerAsset* DataLayerAsset)
 {
 	bool bActorWasModified = false;
-	if (SupportsDataLayer() && DataLayer && !ContainsDataLayer(DataLayer))
+	if (SupportsDataLayer() && DataLayerAsset && !ContainsDataLayer(DataLayerAsset))
 	{
-		if (!bActorWasModified)
-		{
-			Modify();
-			bActorWasModified = true;
-		}
-
-		DataLayers.Emplace(DataLayer->GetFName());
+		Modify();
+		bActorWasModified = true;
+		DataLayerAssets.Add(DataLayerAsset);
 	}
 	return bActorWasModified;
 }
 
-bool AActor::RemoveDataLayer(const UDataLayer* DataLayer)
+bool AActor::RemoveDataLayer(const UDataLayerInstance* DataLayerInstance)
+{
+	return DataLayerInstance->RemoveActor(this);
+}
+
+bool AActor::RemoveDataLayer(const UDataLayerAsset* DataLayerAsset)
 {
 	bool bActorWasModified = false;
-	if (ContainsDataLayer(DataLayer))
+	if (ContainsDataLayer(DataLayerAsset))
 	{
 		if (!bActorWasModified)
 		{
@@ -1343,7 +1348,7 @@ bool AActor::RemoveDataLayer(const UDataLayer* DataLayer)
 			bActorWasModified = true;
 		}
 
-		DataLayers.Remove(FActorDataLayer(DataLayer->GetFName()));
+		DataLayerAssets.Remove(DataLayerAsset);
 	}
 	return bActorWasModified;
 }
@@ -1353,80 +1358,86 @@ bool AActor::RemoveAllDataLayers()
 	if (HasDataLayers())
 	{
 		Modify();
+		DataLayerAssets.Empty();
 		DataLayers.Empty();
 		return true;
 	}
 	return false;
 }
 
-bool AActor::ContainsDataLayer(const UDataLayer* DataLayer) const
+bool AActor::ContainsDataLayer(const UDataLayerInstance* DataLayerInstance) const
+{
+	return DataLayerInstance->ContainsActor(this);
+}
+
+bool AActor::ContainsDataLayer(const UDataLayerAsset* DataLayerAsset) const
 {
-	return DataLayer && DataLayers.Contains(FActorDataLayer(DataLayer->GetFName()));
+	return DataLayerAsset && DataLayerAssets.Contains(DataLayerAsset);
 }
 
 bool AActor::HasDataLayers() const
 {
-	return DataLayers.Num() > 0;
+	return DataLayerAssets.Num() > 0 || DataLayers.Num() > 0;
 }
 
 bool AActor::HasValidDataLayers() const
 {
 	if (const AWorldDataLayers* WorldDataLayers = GetWorld()->GetWorldDataLayers())
 	{
-		for (const FActorDataLayer& DataLayer : DataLayers)
+		for (const TObjectPtr<const UDataLayerAsset>& DataLayerAsset : DataLayerAssets)
 		{
-			if (const UDataLayer* DataLayerObject = WorldDataLayers->GetDataLayerFromName(DataLayer.Name))
+			if (const UDataLayerInstance* DataLayerInstance = WorldDataLayers->GetDataLayerInstance(DataLayerAsset))
 			{
 				return true;
 			}
 		}
-	}
-	return false;
-}
-
-bool AActor::HasAllDataLayers(const TArray<const UDataLayer*>& InDataLayers) const
-{
-	if (DataLayers.Num() < InDataLayers.Num())
-	{
-		return false;
-	}
 
-	for (const UDataLayer* DataLayer : InDataLayers)
-	{
-		if (!ContainsDataLayer(DataLayer))
+		for (const FActorDataLayer& ActorDataLayer : DataLayers)
 		{
-			return false;
+			if (const UDataLayerInstance* DataLayerInstance = WorldDataLayers->GetDataLayerInstance(ActorDataLayer.Name))
+			{
+				return true;
+			}
 		}
 	}
-	return true;
+	return false;
 }
 
-TArray<FName> AActor::GetDataLayerNames() const
+TArray<FName> AActor::GetDataLayerInstanceNames() const
 {
+	TArray<FName> DataLayerInstanceNames;
+	
 	const AWorldDataLayers* WorldDataLayers = GetWorld()->GetWorldDataLayers();
-	return WorldDataLayers ? WorldDataLayers->GetDataLayerNames(DataLayers) : TArray<FName>();
-}
+	if (WorldDataLayers != nullptr)
+	{
+		DataLayerInstanceNames += WorldDataLayers->GetDataLayerInstanceNames(DataLayerAssets);
 
-TArray<const UDataLayer*> AActor::GetDataLayerObjects() const
-{
-	return GetWorld() ? GetDataLayerObjects(GetWorld()->GetWorldDataLayers()) : TArray<const UDataLayer*>();
+		PRAGMA_DISABLE_DEPRECATION_WARNINGS
+		DataLayerInstanceNames += WorldDataLayers->GetDataLayerInstanceNames(DataLayers);
+		PRAGMA_ENABLE_DEPRECATION_WARNINGS
+	}
+	return DataLayerInstanceNames;
 }
 
-TArray<const UDataLayer*> AActor::GetDataLayerObjects(const AWorldDataLayers* WorldDataLayers) const
+TArray<const UDataLayerInstance*> AActor::GetDataLayerInstances() const
 {
-	return WorldDataLayers ? WorldDataLayers->GetDataLayerObjects(DataLayers) : TArray<const UDataLayer*>();
+	return GetWorld() ? GetDataLayerInstances(GetWorld()->GetWorldDataLayers()) : TArray<const UDataLayerInstance*>();
 }
 
-bool AActor::HasAnyOfDataLayers(const TArray<FName>& DataLayerNames) const
+TArray<const UDataLayerInstance*> AActor::GetDataLayerInstances(const AWorldDataLayers* WorldDataLayers) const
 {
-	for (const FActorDataLayer& DataLayer : DataLayers)
+	TArray<const UDataLayerInstance*> DataLayerInstances;
+	if (WorldDataLayers != nullptr)
 	{
-		if (DataLayerNames.Contains(DataLayer.Name))
-		{
-			return true;
-		}
+		DataLayerInstances += WorldDataLayers->GetDataLayerInstances(DataLayerAssets);
+
+		PRAGMA_DISABLE_DEPRECATION_WARNINGS
+		DataLayerInstances += WorldDataLayers->GetDataLayerInstances(DataLayers);
+		PRAGMA_ENABLE_DEPRECATION_WARNINGS
 	}
-	return false;
+
+
+	return DataLayerInstances;
 }
 
 void AActor::FixupDataLayers(bool bRevertChangesOnLockedDataLayer /*= false*/)
@@ -1435,6 +1446,7 @@ void AActor::FixupDataLayers(bool bRevertChangesOnLockedDataLayer /*= false*/)
 	{
 		if (!SupportsDataLayer())
 		{
+			DataLayerAssets.Empty();
 			DataLayers.Empty();
 			return;
 		}
@@ -1446,16 +1458,16 @@ void AActor::FixupDataLayers(bool bRevertChangesOnLockedDataLayer /*= false*/)
 				if (bRevertChangesOnLockedDataLayer)
 				{
 					// Since it's not possible to prevent changes of particular elements of an array, rollback change on locked DataLayers.
-					TSet<FActorDataLayer> PreEdit(PreEditChangeDataLayers);
-					TSet<FActorDataLayer> PostEdit(DataLayers);
+					TSet<const UDataLayerAsset*> PreEdit(PreEditChangeDataLayers);
+					TSet<const UDataLayerAsset*> PostEdit(DataLayerAssets);
 
-					auto DifferenceContainsLockedDataLayers = [WorldDataLayers](const TSet<FActorDataLayer>& A, const TSet<FActorDataLayer>& B)
+					auto DifferenceContainsLockedDataLayers = [WorldDataLayers](const TSet<const UDataLayerAsset*>& A, const TSet<const UDataLayerAsset*>& B)
 					{
-						TSet<FActorDataLayer> Diff = A.Difference(B);
-						for (const FActorDataLayer& ActorDataLayer : Diff)
+						TSet<const UDataLayerAsset*> Diff = A.Difference(B);
+						for (const UDataLayerAsset* DataLayerAsset : Diff)
 						{
-							const UDataLayer* DataLayer = WorldDataLayers->GetDataLayerFromName(ActorDataLayer);
-							if (DataLayer && DataLayer->IsLocked())
+							const UDataLayerInstance* DataLayerInstance = WorldDataLayers->GetDataLayerInstance(DataLayerAsset);
+							if (DataLayerInstance && DataLayerInstance->IsLocked())
 							{
 								return true;
 							}
@@ -1466,24 +1478,34 @@ void AActor::FixupDataLayers(bool bRevertChangesOnLockedDataLayer /*= false*/)
 					if (DifferenceContainsLockedDataLayers(PreEdit, PostEdit) || 
 						DifferenceContainsLockedDataLayers(PostEdit, PreEdit))
 					{
-						DataLayers = PreEditChangeDataLayers;
+						DataLayerAssets = PreEditChangeDataLayers;
 					}
 				}
 
-				TSet<FName> ExistingDataLayers;
-				for (int32 Index = 0; Index < DataLayers.Num();)
+				PRAGMA_DISABLE_DEPRECATION_WARNINGS
+				auto CleanupDataLayers = [this, WorldDataLayers](auto& DataLayerArray)
 				{
-					const FName& DataLayer = DataLayers[Index].Name;
-					if (!WorldDataLayers->GetDataLayerFromName(DataLayer) || ExistingDataLayers.Contains(DataLayer))
-					{
-						DataLayers.RemoveAtSwap(Index);
-					}
-					else
+					using ArrayType = typename TRemoveReference<decltype(DataLayerArray)>::Type;
+					
+					ArrayType ExistingDataLayer;
+					for (int32 Index = 0; Index < DataLayerArray.Num();)
 					{
-						ExistingDataLayers.Add(DataLayer);
-						++Index;
+						auto& DataLayer = DataLayerArray[Index];
+						if (!WorldDataLayers->GetDataLayerInstance(DataLayer) || ExistingDataLayer.Contains(DataLayer))
+						{
+							DataLayerArray.RemoveAtSwap(Index);
+						}
+						else
+						{
+							ExistingDataLayer.Add(DataLayer);
+							++Index;
+						}
 					}
-				}
+				};
+
+				CleanupDataLayers(DataLayerAssets);
+				CleanupDataLayers(DataLayers);
+				PRAGMA_ENABLE_DEPRECATION_WARNINGS
 			}
 		}
 	}
@@ -1493,29 +1515,16 @@ bool AActor::IsPropertyChangedAffectingDataLayers(FPropertyChangedEvent& Propert
 {
 	if (PropertyChangedEvent.Property != nullptr)
 	{
-		FProperty* MemberPropertyThatChanged = PropertyChangedEvent.MemberProperty;
-		const FName MemberPropertyName = MemberPropertyThatChanged != NULL ? MemberPropertyThatChanged->GetFName() : NAME_None;
-
-		static const FName NAME_DataLayers = GET_MEMBER_NAME_CHECKED(AActor, DataLayers);
-		static const FName NAME_FActorDataLayerName = GET_MEMBER_NAME_CHECKED(FActorDataLayer, Name);
+		static const FName NAME_DataLayerAssets = GET_MEMBER_NAME_CHECKED(AActor, DataLayerAssets);
 
-		if (MemberPropertyName == NAME_DataLayers &&
-			PropertyChangedEvent.ChangeType == EPropertyChangeType::ValueSet &&
-			PropertyChangedEvent.Property->GetFName() == NAME_FActorDataLayerName)
+		const FName PropertyName = PropertyChangedEvent.GetPropertyName();
+		if (PropertyName == NAME_DataLayerAssets &&
+			((PropertyChangedEvent.ChangeType == EPropertyChangeType::ValueSet) ||
+				(PropertyChangedEvent.ChangeType == EPropertyChangeType::ArrayClear) ||
+				(PropertyChangedEvent.ChangeType == EPropertyChangeType::Duplicate)))
 		{
 			return true;
 		}
-		else
-		{
-			const FName PropertyName = PropertyChangedEvent.GetPropertyName();
-			if (PropertyName == NAME_DataLayers && 
-				((PropertyChangedEvent.ChangeType == EPropertyChangeType::ValueSet) || 
-				 (PropertyChangedEvent.ChangeType == EPropertyChangeType::ArrayClear) ||
-				 (PropertyChangedEvent.ChangeType == EPropertyChangeType::Duplicate)))
-			{
-				return true;
-			}
-		}
 	}
 	return false;
 }
@@ -1528,7 +1537,7 @@ bool AActor::IsValidForDataLayer() const
 		return false;
 	}
 
-	const bool bIsPartitionedActor = UWorld::IsPartitionedWorld(World);
+	const bool bIsPartitionedActor = UWorld::HasSubsystem<UWorldPartitionSubsystem>(World);
 	const bool bIsInEditorWorld = World->WorldType == EWorldType::Editor;
 	const bool bIsBuilderBrush = FActorEditorUtils::IsABuilderBrush(this);
 	const bool bIsHidden = GetClass()->GetDefaultObject<AActor>()->bHiddenEd;
@@ -1537,6 +1546,61 @@ bool AActor::IsValidForDataLayer() const
 	return bIsValid;
 }
 
+//~ Begin Deprecated
+
+PRAGMA_DISABLE_DEPRECATION_WARNINGS
+
+bool AActor::AddDataLayer(const FActorDataLayer& ActorDataLayer)
+{
+	if (!ContainsDataLayer(ActorDataLayer))
+	{
+		Modify();
+		DataLayers.Add(ActorDataLayer);
+		return true;
+	}
+
+	return false;
+}
+
+bool AActor::RemoveDataLayer(const FActorDataLayer& ActorDataLayer)
+{
+	if (ContainsDataLayer(ActorDataLayer))
+	{
+		Modify();
+		DataLayers.Remove(ActorDataLayer);
+		return true;
+	}
+
+	return false;
+}
+
+bool AActor::AddDataLayer(const UDEPRECATED_DataLayer* DataLayer)
+{
+	if (SupportsDataLayer() && DataLayer)
+	{
+		return AddDataLayer(FActorDataLayer(DataLayer->GetFName()));
+	}
+	return false;
+}
+
+bool AActor::RemoveDataLayer(const UDEPRECATED_DataLayer* DataLayer)
+{
+	if (DataLayer)
+	{
+		return RemoveDataLayer(FActorDataLayer(DataLayer->GetFName()));
+	}
+	return false;
+}
+
+PRAGMA_ENABLE_DEPRECATION_WARNINGS
+
+bool AActor::ContainsDataLayer(const FActorDataLayer& ActorDataLayer) const
+{
+	return DataLayers.Contains(ActorDataLayer);
+}
+
+//~ End Deprecated
+
 // DataLayers (end)
 //---------------------------------------------------------------------------
 
diff --git a/Engine/Source/Runtime/Engine/Private/ActorPartition/ActorPartitionSubsystem.cpp b/Engine/Source/Runtime/Engine/Private/ActorPartition/ActorPartitionSubsystem.cpp
--- a/Engine/Source/Runtime/Engine/Private/ActorPartition/ActorPartitionSubsystem.cpp
+++ b/Engine/Source/Runtime/Engine/Private/ActorPartition/ActorPartitionSubsystem.cpp
@@ -203,7 +203,7 @@ public:
 				(PartitionActorDesc->GridIndexZ == InCellCoord.Z) &&
 				(PartitionActorDesc->GridSize == InGridSize) &&
 				(PartitionActorDesc->GridGuid == InActorPartitionId.GetGridGuid()) &&
-				(FDataLayerEditorContext(ThisWorld, PartitionActorDesc->GetDataLayers()).GetHash() == InActorPartitionId.GetDataLayerEditorContextHash()))
+				(FDataLayerEditorContext(ThisWorld, PartitionActorDesc->GetDataLayerInstanceNames()).GetHash() == InActorPartitionId.GetDataLayerEditorContextHash()))
 			{
 				AActor* DescActor = ActorDesc->GetActor();
 
@@ -395,7 +395,7 @@ APartitionActor* UActorPartitionSubsystem::GetActor(const TSubclassOf<APartition
 	{
 		if (const AWorldDataLayers* WorldDataLayers = World->GetWorldDataLayers())
 		{
-			for (UDataLayer* DataLayer : WorldDataLayers->GetActorEditorContextDataLayers())
+			for (UDataLayerInstance* DataLayer : WorldDataLayers->GetActorEditorContextDataLayers())
 			{
 				PartitionActor->AddDataLayer(DataLayer);
 			}
diff --git a/Engine/Source/Runtime/Engine/Private/EngineTypes.cpp b/Engine/Source/Runtime/Engine/Private/EngineTypes.cpp
--- a/Engine/Source/Runtime/Engine/Private/EngineTypes.cpp
+++ b/Engine/Source/Runtime/Engine/Private/EngineTypes.cpp
@@ -253,13 +253,13 @@ FActorInstanceHandle::FActorInstanceHandle(AActor* InActor)
 	if (InActor)
 	{
 #if WITH_EDITOR
-		// use the first layer the actor is in if it's in multiple layers
-		TArray<const UDataLayer*> DataLayers = InActor->GetDataLayerObjects();
-		const UDataLayer* Layer = DataLayers.Num() > 0 ? DataLayers[0] : nullptr;
+		// use the first data layer the actor is in if it's in multiple layers
+		TArray<const UDataLayerInstance*> DataLayerInstances = InActor->GetDataLayerInstances();
+		const UDataLayerInstance* DataLayerInstance = DataLayerInstances.Num() > 0 ? DataLayerInstances[0] : nullptr;
 #else
-		const UDataLayer* Layer = nullptr;
+		const UDataLayerInstance* DataLayerInstance = nullptr;
 #endif // WITH_EDITOR
-		Manager = FLightWeightInstanceSubsystem::Get().FindLightWeightInstanceManager(InActor->StaticClass(), Layer);
+		Manager = FLightWeightInstanceSubsystem::Get().FindLightWeightInstanceManager(InActor->StaticClass(), DataLayerInstance);
 	}
 }
 
@@ -516,13 +516,13 @@ bool FActorInstanceHandle::operator==(const AActor* OtherActor) const
 
 #if WITH_EDITOR
 	// use the first layer the actor is in if it's in multiple layers
-	TArray<const UDataLayer*> DataLayers = OtherActor->GetDataLayerObjects();
-	const UDataLayer* Layer = DataLayers.Num() > 0 ? DataLayers[0] : nullptr;
+	TArray<const UDataLayerInstance*> DataLayerInstances = OtherActor->GetDataLayerInstances();
+	const UDataLayerInstance* DataLayerInstance = DataLayerInstances.Num() > 0 ? DataLayerInstances[0] : nullptr;
 #else
-	const UDataLayer* Layer = nullptr;
+	const UDataLayerInstance* DataLayerInstance = nullptr;
 #endif // WITH_EDITOR
 
-	if (ALightWeightInstanceManager* LWIManager = FLightWeightInstanceSubsystem::Get().FindLightWeightInstanceManager(OtherActor->StaticClass(), Layer))
+	if (ALightWeightInstanceManager* LWIManager = FLightWeightInstanceSubsystem::Get().FindLightWeightInstanceManager(OtherActor->StaticClass(), DataLayerInstance))
 	{
 		if (Manager.Get() != LWIManager)
 		{
diff --git a/Engine/Source/Runtime/Engine/Private/LevelInstance/LevelInstanceSubsystem.cpp b/Engine/Source/Runtime/Engine/Private/LevelInstance/LevelInstanceSubsystem.cpp
--- a/Engine/Source/Runtime/Engine/Private/LevelInstance/LevelInstanceSubsystem.cpp
+++ b/Engine/Source/Runtime/Engine/Private/LevelInstance/LevelInstanceSubsystem.cpp
@@ -1049,10 +1049,10 @@ void ULevelInstanceSubsystem::BreakLevelInstance_Impl(ALevelInstance* LevelInsta
 			return;
 		}
 
-		TArray<const UDataLayer*> LevelInstanceDataLayers = LevelInstanceActor->GetDataLayerObjects();
+		TArray<const UDataLayerInstance*> LevelInstanceDataLayerInstances = LevelInstanceActor->GetDataLayerInstances();
 
 		TSet<AActor*> ActorsToMove;
-		TFunction<bool(AActor*)> AddActorToMove = [this, &ActorsToMove, &AddActorToMove, &LevelInstanceDataLayers](AActor* Actor)
+		TFunction<bool(AActor*)> AddActorToMove = [this, &ActorsToMove, &AddActorToMove, &LevelInstanceDataLayerInstances](AActor* Actor)
 		{
 			if (ActorsToMove.Contains(Actor))
 			{
@@ -1078,9 +1078,9 @@ void ULevelInstanceSubsystem::BreakLevelInstance_Impl(ALevelInstance* LevelInsta
 					// Apply the same data layer settings to the actors to move out
 					if (Actor->SupportsDataLayer() && Actor->IsValidForDataLayer())
 					{
-						for (const UDataLayer* DataLayer : LevelInstanceDataLayers)
+						for (const UDataLayerInstance* DataLayerInstance : LevelInstanceDataLayerInstances)
 						{
-							Actor->AddDataLayer(DataLayer);
+							Actor->AddDataLayer(DataLayerInstance);
 						}
 					}
 
diff --git a/Engine/Source/Runtime/Engine/Private/LightWeightInstanceBlueprintFunctionLibrary.cpp b/Engine/Source/Runtime/Engine/Private/LightWeightInstanceBlueprintFunctionLibrary.cpp
--- a/Engine/Source/Runtime/Engine/Private/LightWeightInstanceBlueprintFunctionLibrary.cpp
+++ b/Engine/Source/Runtime/Engine/Private/LightWeightInstanceBlueprintFunctionLibrary.cpp
@@ -4,7 +4,7 @@
 #include "GameFramework/LightWeightInstanceSubsystem.h"
 #include "GameFramework/LightWeightInstanceManager.h"
 
-FActorInstanceHandle ULightWeightInstanceBlueprintFunctionLibrary::CreateNewLightWeightInstance(UClass* InActorClass, FTransform InTransform, UDataLayer* InLayer, UWorld* World)
+FActorInstanceHandle ULightWeightInstanceBlueprintFunctionLibrary::CreateNewLightWeightInstance(UClass* InActorClass, FTransform InTransform, UDataLayerInstance* InLayer, UWorld* World)
 {
 	// set up initialization data
 	FLWIData PerInstanceData;
@@ -23,12 +23,12 @@ FActorInstanceHandle ULightWeightInstanceBlueprintFunctionLibrary::ConvertActorT
 	// Get or create a light weight instance for this class and layer
 	// use the first layer the actor is in if it's in multiple layers
 #if WITH_EDITOR
-	TArray<const UDataLayer*> DataLayers = InActor->GetDataLayerObjects();
-	const UDataLayer* Layer = DataLayers.Num() > 0 ? DataLayers[0] : nullptr;
+	TArray<const UDataLayerInstance*> DataLayerInstances = InActor->GetDataLayerInstances();
+	const UDataLayerInstance* DataLayerInstance = DataLayerInstances.Num() > 0 ? DataLayerInstances[0] : nullptr;
 #else
-	const UDataLayer* Layer = nullptr;
+	const UDataLayerInstance* DataLayerInstance = nullptr;
 #endif // WITH_EDITOR
-	if (ALightWeightInstanceManager* LWIManager = FLightWeightInstanceSubsystem::Get().FindOrAddLightWeightInstanceManager(InActor->GetClass(), Layer, InActor->GetWorld()))
+	if (ALightWeightInstanceManager* LWIManager = FLightWeightInstanceSubsystem::Get().FindOrAddLightWeightInstanceManager(InActor->GetClass(), DataLayerInstance, InActor->GetWorld()))
 	{
 		return LWIManager->ConvertActorToLightWeightInstance(InActor);
 	}
diff --git a/Engine/Source/Runtime/Engine/Private/LightWeightInstanceSubsystem.cpp b/Engine/Source/Runtime/Engine/Private/LightWeightInstanceSubsystem.cpp
--- a/Engine/Source/Runtime/Engine/Private/LightWeightInstanceSubsystem.cpp
+++ b/Engine/Source/Runtime/Engine/Private/LightWeightInstanceSubsystem.cpp
@@ -4,7 +4,7 @@
 #include "Engine/Engine.h"
 #include "Engine/World.h"
 #include "UObject/UObjectIterator.h"
-#include "WorldPartition/DataLayer/DataLayer.h"
+#include "WorldPartition/DataLayer/DataLayerInstance.h"
 
 DEFINE_LOG_CATEGORY(LogLightWeightInstance);
 
@@ -78,11 +78,11 @@ ALightWeightInstanceManager* FLightWeightInstanceSubsystem::FindLightWeightInsta
 			{
 #if WITH_EDITOR
 				// make sure the data layers match
-				TArray<const UDataLayer*> ManagerLayers = LWInstance->GetDataLayerObjects();
-				const UDataLayer* ManagerLayer = ManagerLayers.Num() > 0 ? ManagerLayers[0] : nullptr;
+				TArray<const UDataLayerInstance*> ManagerLayers = LWInstance->GetDataLayerInstances();
+				const UDataLayerInstance* ManagerLayer = ManagerLayers.Num() > 0 ? ManagerLayers[0] : nullptr;
 
-				TArray<const UDataLayer*> ActorLayers = LWInstance->GetDataLayerObjects();
-				const UDataLayer* ActorLayer = ActorLayers.Num() > 0 ? ActorLayers[0] : nullptr;
+				TArray<const UDataLayerInstance*> ActorLayers = LWInstance->GetDataLayerInstances();
+				const UDataLayerInstance* ActorLayer = ActorLayers.Num() > 0 ? ActorLayers[0] : nullptr;
 
 				if (ManagerLayer == ActorLayer)
 #endif // WITH_EDITOR
@@ -96,7 +96,7 @@ ALightWeightInstanceManager* FLightWeightInstanceSubsystem::FindLightWeightInsta
 	return nullptr;
 }
 
-ALightWeightInstanceManager* FLightWeightInstanceSubsystem::FindLightWeightInstanceManager(UClass* ActorClass, const UDataLayer* Layer) const
+ALightWeightInstanceManager* FLightWeightInstanceSubsystem::FindLightWeightInstanceManager(UClass* ActorClass, const UDataLayerInstance* Layer) const
 {
 	if (ActorClass == nullptr)
 	{
@@ -120,7 +120,7 @@ ALightWeightInstanceManager* FLightWeightInstanceSubsystem::FindLightWeightInsta
 	return nullptr;
 }
 
-ALightWeightInstanceManager* FLightWeightInstanceSubsystem::FindOrAddLightWeightInstanceManager(UClass* ActorClass, const UDataLayer* Layer, UWorld* World)
+ALightWeightInstanceManager* FLightWeightInstanceSubsystem::FindOrAddLightWeightInstanceManager(UClass* ActorClass, const UDataLayerInstance* Layer, UWorld* World)
 {
 	if (ActorClass == nullptr || World == nullptr)
 	{
@@ -330,7 +330,7 @@ bool FLightWeightInstanceSubsystem::IsInLevel(const FActorInstanceHandle& Handle
 	return false;
 }
 
-FActorInstanceHandle FLightWeightInstanceSubsystem::CreateNewLightWeightInstance(UClass* InActorClass, FLWIData* InitData, UDataLayer* InLayer, UWorld* World)
+FActorInstanceHandle FLightWeightInstanceSubsystem::CreateNewLightWeightInstance(UClass* InActorClass, FLWIData* InitData, UDataLayerInstance* InLayer, UWorld* World)
 {
 	// Get or create a light weight instance for this class and data layer
 	if (ALightWeightInstanceManager* LWIManager = FindOrAddLightWeightInstanceManager(InActorClass, InLayer, World))
diff --git a/Engine/Source/Runtime/Engine/Private/WorldPartition/DataLayer/DataLayer.cpp b/Engine/Source/Runtime/Engine/Private/WorldPartition/DataLayer/DataLayer.cpp
--- a/Engine/Source/Runtime/Engine/Private/WorldPartition/DataLayer/DataLayer.cpp
+++ b/Engine/Source/Runtime/Engine/Private/WorldPartition/DataLayer/DataLayer.cpp
@@ -1,7 +1,7 @@
 // Copyright Epic Games, Inc. All Rights Reserved.
 
 /*=============================================================================
-	DataLayer.cpp: UDataLayer class implementation
+	DataLayer.cpp: UDEPRECATED_DataLayer class implementation
 =============================================================================*/
 
 #include "WorldPartition/DataLayer/DataLayer.h"
@@ -9,7 +9,9 @@
 
 #define LOCTEXT_NAMESPACE "DataLayer"
 
-UDataLayer::UDataLayer(const FObjectInitializer& ObjectInitializer)
+PRAGMA_DISABLE_DEPRECATION_WARNINGS
+
+UDEPRECATED_DataLayer::UDEPRECATED_DataLayer(const FObjectInitializer& ObjectInitializer)
 : Super(ObjectInitializer)
 #if WITH_EDITORONLY_DATA
 , bIsInitiallyActive_DEPRECATED(false)
@@ -26,7 +28,7 @@ UDataLayer::UDataLayer(const FObjectInitializer& ObjectInitializer)
 {
 }
 
-void UDataLayer::PostLoad()
+void UDEPRECATED_DataLayer::PostLoad()
 {
 	Super::PostLoad();
 
@@ -40,7 +42,7 @@ void UDataLayer::PostLoad()
 	bIsVisible = bIsInitiallyVisible;
 
 	// Sanitize Label
-	DataLayerLabel = UDataLayer::GetSanitizedDataLayerLabel(DataLayerLabel);
+	DataLayerLabel = DataLayerUtils::GetSanitizedDataLayerLabel(DataLayerLabel);
 
 	if (DebugColor == FColor::Black)
 	{
@@ -52,13 +54,13 @@ void UDataLayer::PostLoad()
 	}
 #endif
 
-	if (Parent)
+	if (Parent_DEPRECATED)
 	{
-		Parent->AddChild(this);
+		Parent_DEPRECATED->AddChild(this);
 	}
 }
 
-bool UDataLayer::IsInitiallyVisible() const
+bool UDEPRECATED_DataLayer::IsInitiallyVisible() const
 {
 #if WITH_EDITOR
 	return bIsInitiallyVisible;
@@ -67,7 +69,7 @@ bool UDataLayer::IsInitiallyVisible() const
 #endif
 }
 
-bool UDataLayer::IsVisible() const
+bool UDEPRECATED_DataLayer::IsVisible() const
 {
 #if WITH_EDITOR
 	return bIsVisible;
@@ -76,17 +78,11 @@ bool UDataLayer::IsVisible() const
 #endif
 }
 
-FName UDataLayer::GetSanitizedDataLayerLabel(FName InDataLayerLabel)
-{
-	// Removes all quotes as well as whitespace characters from the startand end
-	return FName(InDataLayerLabel.ToString().TrimStartAndEnd().Replace(TEXT("\""), TEXT("")));
-}
-
-bool UDataLayer::IsEffectiveVisible() const
+bool UDEPRECATED_DataLayer::IsEffectiveVisible() const
 {
 #if WITH_EDITOR
 	bool bResult = IsVisible();
-	const UDataLayer* ParentDataLayer = GetParent();
+	const UDEPRECATED_DataLayer* ParentDataLayer = GetParent();
 	while (ParentDataLayer && bResult)
 	{
 		bResult = bResult && ParentDataLayer->IsVisible();
@@ -98,11 +94,11 @@ bool UDataLayer::IsEffectiveVisible() const
 #endif
 }
 
-void UDataLayer::AddChild(UDataLayer* InDataLayer)
+void UDEPRECATED_DataLayer::AddChild(UDEPRECATED_DataLayer* InDataLayer)
 {
 	Modify();
-	checkSlow(!Children.Contains(InDataLayer));
-	Children.Add(InDataLayer);
+	checkSlow(!Children_DEPRECATED.Contains(InDataLayer));
+	Children_DEPRECATED.Add(InDataLayer);
 #if WITH_EDITOR
 	if (IsRuntime())
 	{
@@ -113,10 +109,10 @@ void UDataLayer::AddChild(UDataLayer* InDataLayer)
 
 #if WITH_EDITOR
 
-bool UDataLayer::IsEffectiveLoadedInEditor() const
+bool UDEPRECATED_DataLayer::IsEffectiveLoadedInEditor() const
 {
 	bool bResult = IsLoadedInEditor();
-	const UDataLayer* ParentDataLayer = GetParent();
+	const UDEPRECATED_DataLayer* ParentDataLayer = GetParent();
 	while (ParentDataLayer && bResult)
 	{
 		bResult = bResult && ParentDataLayer->IsLoadedInEditor();
@@ -125,7 +121,7 @@ bool UDataLayer::IsEffectiveLoadedInEditor() const
 	return bResult;
 }
 
-bool UDataLayer::IsLocked() const
+bool UDEPRECATED_DataLayer::IsLocked() const
 {
 	if (bIsLocked)
 	{
@@ -135,16 +131,16 @@ bool UDataLayer::IsLocked() const
 	return IsRuntime() && !GetOuterAWorldDataLayers()->GetAllowRuntimeDataLayerEditing();
 }
 
-bool UDataLayer::CanEditChange(const FProperty* InProperty) const
+bool UDEPRECATED_DataLayer::CanEditChange(const FProperty* InProperty) const
 {
-	if ((InProperty->GetFName() == GET_MEMBER_NAME_CHECKED(UDataLayer, bIsRuntime)) ||
-		(InProperty->GetFName() == GET_MEMBER_NAME_CHECKED(UDataLayer, InitialRuntimeState)) ||
-		(InProperty->GetFName() == GET_MEMBER_NAME_CHECKED(UDataLayer, DebugColor)))
+	if ((InProperty->GetFName() == GET_MEMBER_NAME_CHECKED(UDEPRECATED_DataLayer, bIsRuntime)) ||
+		(InProperty->GetFName() == GET_MEMBER_NAME_CHECKED(UDEPRECATED_DataLayer, InitialRuntimeState)) ||
+		(InProperty->GetFName() == GET_MEMBER_NAME_CHECKED(UDEPRECATED_DataLayer, DebugColor)))
 	{
-		if (InProperty->GetFName() == GET_MEMBER_NAME_CHECKED(UDataLayer, bIsRuntime))
+		if (InProperty->GetFName() == GET_MEMBER_NAME_CHECKED(UDEPRECATED_DataLayer, bIsRuntime))
 		{
 			// If DataLayer is Runtime because of its parent being Runtime, we don't allow modifying 
-			if (Parent && Parent->IsRuntime())
+			if (Parent_DEPRECATED && Parent_DEPRECATED->IsRuntime())
 			{
 				check(IsRuntime());
 				return false;
@@ -156,9 +152,9 @@ bool UDataLayer::CanEditChange(const FProperty* InProperty) const
 	return Super::CanEditChange(InProperty);
 }
 
-void UDataLayer::PostEditChangeProperty(FPropertyChangedEvent& PropertyChangedEvent)
+void UDEPRECATED_DataLayer::PostEditChangeProperty(FPropertyChangedEvent& PropertyChangedEvent)
 {
-	static const FName NAME_IsRuntime = GET_MEMBER_NAME_CHECKED(UDataLayer, bIsRuntime);
+	static const FName NAME_IsRuntime = GET_MEMBER_NAME_CHECKED(UDEPRECATED_DataLayer, bIsRuntime);
 	FProperty* MemberPropertyThatChanged = PropertyChangedEvent.MemberProperty;
 	const FName MemberPropertyName = MemberPropertyThatChanged != NULL ? MemberPropertyThatChanged->GetFName() : NAME_None;
 	if (MemberPropertyName == NAME_IsRuntime)
@@ -168,12 +164,12 @@ void UDataLayer::PostEditChangeProperty(FPropertyChangedEvent& PropertyChangedEv
 	Super::PostEditChangeProperty(PropertyChangedEvent);
 }
 
-bool UDataLayer::CanParent(const UDataLayer* InParent) const
+bool UDEPRECATED_DataLayer::CanParent(const UDEPRECATED_DataLayer* InParent) const
 {
-	return (this != InParent) && (Parent != InParent);
+	return (this != InParent) && (Parent_DEPRECATED != InParent);
 }
 
-void UDataLayer::SetParent(UDataLayer* InParent)
+void UDEPRECATED_DataLayer::SetParent(UDEPRECATED_DataLayer* InParent)
 {
 	if (!CanParent(InParent))
 	{
@@ -181,18 +177,18 @@ void UDataLayer::SetParent(UDataLayer* InParent)
 	}
 
 	Modify();
-	if (Parent)
+	if (Parent_DEPRECATED)
 	{
-		Parent->RemoveChild(this);
+		Parent_DEPRECATED->RemoveChild(this);
 	}
-	Parent = InParent;
-	if (Parent)
+	Parent_DEPRECATED = InParent;
+	if (Parent_DEPRECATED)
 	{
-		Parent->AddChild(this);
+		Parent_DEPRECATED->AddChild(this);
 	}
 }
 
-void UDataLayer::SetChildParent(UDataLayer* InParent)
+void UDEPRECATED_DataLayer::SetChildParent(UDEPRECATED_DataLayer* InParent)
 {
 	if (this == InParent)
 	{
@@ -200,37 +196,25 @@ void UDataLayer::SetChildParent(UDataLayer* InParent)
 	}
 
 	Modify();
-	while (Children.Num())
+	while (Children_DEPRECATED.Num())
 	{
-		Children[0]->SetParent(InParent);
+		Children_DEPRECATED[0]->SetParent(InParent);
 	};
 }
 
-void UDataLayer::RemoveChild(UDataLayer* InDataLayer)
+void UDEPRECATED_DataLayer::RemoveChild(UDEPRECATED_DataLayer* InDataLayer)
 {
 	Modify();
-	check(Children.Contains(InDataLayer));
-	Children.RemoveSingle(InDataLayer);
+	check(Children_DEPRECATED.Contains(InDataLayer));
+	Children_DEPRECATED.RemoveSingle(InDataLayer);
 }
 
-const TCHAR* UDataLayer::GetDataLayerIconName() const
+const TCHAR* UDEPRECATED_DataLayer::GetDataLayerIconName() const
 {
 	return IsRuntime() ? TEXT("DataLayer.Runtime") : TEXT("DataLayer.Editor");
 }
 
-void UDataLayer::SetDataLayerLabel(FName InDataLayerLabel)
-{
-	FName DataLayerLabelSanitized = UDataLayer::GetSanitizedDataLayerLabel(InDataLayerLabel);
-	if (DataLayerLabel != DataLayerLabelSanitized)
-	{
-		Modify();
-		AWorldDataLayers* WorldDataLayers = GetOuterAWorldDataLayers();
-		check(!WorldDataLayers || !WorldDataLayers->GetDataLayerFromLabel(DataLayerLabelSanitized))
-		DataLayerLabel = DataLayerLabelSanitized;
-	}
-}
-
-void UDataLayer::SetVisible(bool bInIsVisible)
+void UDEPRECATED_DataLayer::SetVisible(bool bInIsVisible)
 {
 	if (bIsVisible != bInIsVisible)
 	{
@@ -239,7 +223,7 @@ void UDataLayer::SetVisible(bool bInIsVisible)
 	}
 }
 
-void UDataLayer::SetIsInitiallyVisible(bool bInIsInitiallyVisible)
+void UDEPRECATED_DataLayer::SetIsInitiallyVisible(bool bInIsInitiallyVisible)
 {
 	if (bIsInitiallyVisible != bInIsInitiallyVisible)
 	{
@@ -248,7 +232,7 @@ void UDataLayer::SetIsInitiallyVisible(bool bInIsInitiallyVisible)
 	}
 }
 
-void UDataLayer::SetIsRuntime(bool bInIsRuntime)
+void UDEPRECATED_DataLayer::SetIsRuntime(bool bInIsRuntime)
 {
 	if (bIsRuntime != bInIsRuntime)
 	{
@@ -259,35 +243,18 @@ void UDataLayer::SetIsRuntime(bool bInIsRuntime)
 	}
 }
 
-#if WITH_EDITOR
-bool UDataLayer::IsInActorEditorContext() const
-{
-	return GetOuterAWorldDataLayers()->IsInActorEditorContext(this);
-}
-
-bool UDataLayer::AddToActorEditorContext()
-{
-	return GetOuterAWorldDataLayers()->AddToActorEditorContext(this);
-}
-
-bool UDataLayer::RemoveFromActorEditorContext()
-{
-	return GetOuterAWorldDataLayers()->RemoveFromActorEditorContext(this);
-}
-#endif
-
-void UDataLayer::PropagateIsRuntime()
+void UDEPRECATED_DataLayer::PropagateIsRuntime()
 {
 	if (IsRuntime())
 	{
-		for (UDataLayer* Child : Children)
+		for (UDEPRECATED_DataLayer* Child : Children_DEPRECATED)
 		{
 			Child->SetIsRuntime(true);
 		}
 	}
 }
 
-void UDataLayer::SetIsLoadedInEditor(bool bInIsLoadedInEditor, bool bInFromUserChange)
+void UDEPRECATED_DataLayer::SetIsLoadedInEditor(bool bInIsLoadedInEditor, bool bInFromUserChange)
 {
 	if (bIsLoadedInEditor != bInIsLoadedInEditor)
 	{
@@ -297,16 +264,16 @@ void UDataLayer::SetIsLoadedInEditor(bool bInIsLoadedInEditor, bool bInFromUserC
 	}
 }
 
-FText UDataLayer::GetDataLayerText(const UDataLayer* InDataLayer)
+FText UDEPRECATED_DataLayer::GetDataLayerText(const UDEPRECATED_DataLayer* InDataLayer)
 {
 	return InDataLayer ? FText::FromName(InDataLayer->GetDataLayerLabel()) : LOCTEXT("InvalidDataLayerLabel", "<None>");
 }
 
 #endif
 
-void UDataLayer::ForEachChild(TFunctionRef<bool(const UDataLayer*)> Operation) const
+void UDEPRECATED_DataLayer::ForEachChild(TFunctionRef<bool(const UDEPRECATED_DataLayer*)> Operation) const
 {
-	for (UDataLayer* Child : Children)
+	for (UDEPRECATED_DataLayer* Child : Children_DEPRECATED)
 	{
 		if (!Operation(Child))
 		{
@@ -315,4 +282,6 @@ void UDataLayer::ForEachChild(TFunctionRef<bool(const UDataLayer*)> Operation) c
 	}
 }
 
+PRAGMA_ENABLE_DEPRECATION_WARNINGS
+
 #undef LOCTEXT_NAMESPACE
\ No newline at end of file
diff --git a/Engine/Source/Runtime/Engine/Private/WorldPartition/DataLayer/DataLayerEditorContext.cpp b/Engine/Source/Runtime/Engine/Private/WorldPartition/DataLayer/DataLayerEditorContext.cpp
--- a/Engine/Source/Runtime/Engine/Private/WorldPartition/DataLayer/DataLayerEditorContext.cpp
+++ b/Engine/Source/Runtime/Engine/Private/WorldPartition/DataLayer/DataLayerEditorContext.cpp
@@ -4,7 +4,7 @@
 
 #if WITH_EDITOR
 
-#include "WorldPartition/DataLayer/DataLayer.h"
+#include "WorldPartition/DataLayer/DataLayerInstance.h"
 #include "WorldPartition/DataLayer/WorldDataLayers.h"
 #include "WorldPartition/DataLayer/DataLayerSubsystem.h"
 #include "Engine/World.h"
@@ -12,7 +12,7 @@
 /*
  * FDataLayerEditorContext
  */
-FDataLayerEditorContext::FDataLayerEditorContext(UWorld* InWorld, const TArray<FName>& InDataLayers)
+FDataLayerEditorContext::FDataLayerEditorContext(UWorld* InWorld, const TArray<FName>& InDataLayerInstances)
 	: Hash(FDataLayerEditorContext::EmptyHash)
 {
 	const AWorldDataLayers* WorldDataLayers = InWorld->GetWorldDataLayers();
@@ -21,20 +21,20 @@ FDataLayerEditorContext::FDataLayerEditorContext(UWorld* InWorld, const TArray<F
 		return;
 	}
 
-	for (const FName& DataLayerName : InDataLayers)
+	for (const FName& DataLayerInstanceName : InDataLayerInstances)
 	{
-		if (const UDataLayer* DataLayerObject = WorldDataLayers->GetDataLayerFromName(DataLayerName))
+		if (const UDataLayerInstance* DataLayerInstance = WorldDataLayers->GetDataLayerInstance(DataLayerInstanceName))
 		{
-			DataLayers.AddUnique(DataLayerObject->GetFName());
+			DataLayerInstances.AddUnique(DataLayerInstance->GetDataLayerFName());
 		}
 	}
 
-	if (DataLayers.Num())
+	if (DataLayerInstances.Num())
 	{
-		DataLayers.Sort([](const FName& A, const FName& B) { return A.ToString() < B.ToString(); });
-		for (FName LayerName : DataLayers)
+		DataLayerInstances.Sort([](const FName& A, const FName& B) { return A.ToString() < B.ToString(); });
+		for (FName InstanceName : DataLayerInstances)
 		{
-			Hash = FCrc::StrCrc32(*LayerName.ToString(), Hash);
+			Hash = FCrc::StrCrc32(*InstanceName.ToString(), Hash);
 		}
 		check(Hash != FDataLayerEditorContext::EmptyHash);
 	}
diff --git a/Engine/Source/Runtime/Engine/Private/WorldPartition/DataLayer/DataLayerInstanceWithAsset.cpp b/Engine/Source/Runtime/Engine/Private/WorldPartition/DataLayer/DataLayerInstanceWithAsset.cpp
new file mode 100644
--- /dev/null
+++ b/Engine/Source/Runtime/Engine/Private/WorldPartition/DataLayer/DataLayerInstanceWithAsset.cpp
@@ -0,0 +1,79 @@
+// Copyright Epic Games, Inc. All Rights Reserved.
+
+#include "WorldPartition/DataLayer/DataLayerInstanceWithAsset.h"
+
+#include "WorldPartition/DataLayer/WorldDataLayers.h"
+#include "WorldPartition/ErrorHandling/WorldPartitionStreamingGenerationErrorHandler.h"
+
+UDataLayerInstanceWithAsset::UDataLayerInstanceWithAsset(const FObjectInitializer& ObjectInitializer)
+	: Super(ObjectInitializer)
+{
+
+}
+
+#if WITH_EDITOR
+FName UDataLayerInstanceWithAsset::MakeName(const UDataLayerAsset* DeprecatedDataLayer)
+{
+	return FName(FString::Format(TEXT("DataLayer_{0}"), { FGuid::NewGuid().ToString() }));
+}
+
+void UDataLayerInstanceWithAsset::OnCreated(const UDataLayerAsset* Asset)
+{
+	check(!GetOuterAWorldDataLayers()->HasDeprecatedDataLayers() || IsRunningCommandlet());
+
+	Modify(/*bAlwaysMarkDirty*/false);
+
+	check(DataLayerAsset == nullptr);
+	DataLayerAsset = Asset;
+
+	SetVisible(true);
+}
+
+bool UDataLayerInstanceWithAsset::AddActor(AActor* Actor) const
+{
+	return Actor->AddDataLayer(DataLayerAsset);
+}
+
+bool UDataLayerInstanceWithAsset::RemoveActor(AActor* Actor) const
+{
+	return Actor->RemoveDataLayer(DataLayerAsset);
+}
+
+bool UDataLayerInstanceWithAsset::ContainsActor(const AActor* Actor) const
+{
+	return Actor->ContainsDataLayer(DataLayerAsset);
+}
+
+bool UDataLayerInstanceWithAsset::Validate(IStreamingGenerationErrorHandler* ErrorHandler) const
+{
+	bool bIsValid = true;
+
+	if (GetAsset() == nullptr)
+	{
+		ErrorHandler->OnInvalidReferenceDataLayerAsset(this);
+		return false;
+	}
+
+	GetOuterAWorldDataLayers()->ForEachDataLayer([&bIsValid, this, ErrorHandler](UDataLayerInstance* DataLayerInstance)
+	{
+		if (DataLayerInstance != this)
+		{
+			if (UDataLayerInstanceWithAsset* DataLayerInstanceWithAsset = Cast<UDataLayerInstanceWithAsset>(DataLayerInstance))
+			{
+				if (DataLayerInstanceWithAsset->GetAsset() == GetAsset())
+				{
+					ErrorHandler->OnDataLayerAssetConflict(this, DataLayerInstanceWithAsset);
+					bIsValid = false;
+					return false;
+				}
+			}
+		}
+
+		return true;
+	});
+
+	bIsValid &= Super::Validate(ErrorHandler);
+
+	return bIsValid;
+}
+#endif
\ No newline at end of file
