Make DebugColor accessible for Editor Data Layers
Move DebugColor assignement code in DataLayerFactory
Debug Color now uses the DataLayerAsset full path instead of its name to avoid clashing

diff --git a/Engine/Source/Editor/DataLayerEditor/Private/DataLayer/DataLayerFactory.cpp b/Engine/Source/Editor/DataLayerEditor/Private/DataLayer/DataLayerFactory.cpp
--- a/Engine/Source/Editor/DataLayerEditor/Private/DataLayer/DataLayerFactory.cpp
+++ b/Engine/Source/Editor/DataLayerEditor/Private/DataLayer/DataLayerFactory.cpp
@@ -3,6 +3,7 @@
 #include "DataLayer/DataLayerFactory.h"
 
 #include "WorldPartition/DataLayer/DataLayerAsset.h"
+#include "Math/RandomStream.h"
 
 UDataLayerFactory::UDataLayerFactory(const FObjectInitializer& ObjectInitializer)
 	: Super(ObjectInitializer)
@@ -16,5 +17,13 @@ UDataLayerFactory::UDataLayerFactory(const FObjectInitializer& ObjectInitializer
 
 UObject* UDataLayerFactory::FactoryCreateNew(UClass* Class, UObject* InParent, FName InName, EObjectFlags Flags, UObject* Context, FFeedbackContext* Warn)
 {
-	return NewObject<UDataLayerAsset>(InParent, InName, Flags);
+	UDataLayerAsset* DataLayerAsset =   NewObject<UDataLayerAsset>(InParent, InName, Flags);
+
+	FRandomStream RandomStream(FName(DataLayerAsset->GetFullName()));
+	const uint8 R = (uint8)(RandomStream.GetFraction() * 255.f);
+	const uint8 G = (uint8)(RandomStream.GetFraction() * 255.f);
+	const uint8 B = (uint8)(RandomStream.GetFraction() * 255.f);
+	DataLayerAsset->SetDebugColor(FColor(R, G, B));
+
+	return DataLayerAsset;
 }
\ No newline at end of file
diff --git a/Engine/Source/Runtime/Engine/Private/WorldPartition/DataLayer/DataLayerAsset.cpp b/Engine/Source/Runtime/Engine/Private/WorldPartition/DataLayer/DataLayerAsset.cpp
--- a/Engine/Source/Runtime/Engine/Private/WorldPartition/DataLayer/DataLayerAsset.cpp
+++ b/Engine/Source/Runtime/Engine/Private/WorldPartition/DataLayer/DataLayerAsset.cpp
@@ -9,19 +9,3 @@ UDataLayerAsset::UDataLayerAsset(const FObjectInitializer& ObjectInitializer)
 {
 
 }
-
-void UDataLayerAsset::PostLoad()
-{
-	Super::PostLoad();
-
-#if WITH_EDITORONLY_DATA
-	if (DebugColor == FColor::Black)
-	{
-		FRandomStream RandomStream(GetFName());
-		const uint8 R = (uint8)(RandomStream.GetFraction() * 255.f);
-		const uint8 G = (uint8)(RandomStream.GetFraction() * 255.f);
-		const uint8 B = (uint8)(RandomStream.GetFraction() * 255.f);
-		DebugColor = FColor(R, G, B);
-	}
-#endif
-}
\ No newline at end of file
diff --git a/Engine/Source/Runtime/Engine/Public/WorldPartition/DataLayer/DataLayerAsset.h b/Engine/Source/Runtime/Engine/Public/WorldPartition/DataLayer/DataLayerAsset.h
--- a/Engine/Source/Runtime/Engine/Public/WorldPartition/DataLayer/DataLayerAsset.h
+++ b/Engine/Source/Runtime/Engine/Public/WorldPartition/DataLayer/DataLayerAsset.h
@@ -19,10 +19,9 @@ class ENGINE_API UDataLayerAsset : public UObject
 	friend class UDataLayerConversionInfo;
 
 public:
-	virtual void PostLoad() override;
-
 #if WITH_EDITOR
 	void SetType(EDataLayerType Type) { DataLayerType = Type; }
+	void SetDebugColor(FColor InDebugColor) { DebugColor = InDebugColor; }
 #endif
 
 	UFUNCTION(Category = "Data Layer", BlueprintCallable)
@@ -39,6 +38,6 @@ private:
 	UPROPERTY(Category = "Data Layer", EditAnywhere)
 	EDataLayerType DataLayerType;
 
-	UPROPERTY(Category = "Data Layer|Runtime", EditAnywhere, meta = (EditConditionHides, EditCondition = "DataLayerType == EDataLayerType::Runtime"))
+	UPROPERTY(Category = "Data Layer|Runtime", EditAnywhere)
 	FColor DebugColor;
 };
\ No newline at end of file
