Give a clearer title to the dialog opened when creating a Data Layer Instance
" Pick Data Layer Asset" dialog now remmembers the last path where a data layer was picked

diff --git a/Engine/Source/Editor/DataLayerEditor/Private/DataLayer/DataLayerEditorSubsystem.cpp b/Engine/Source/Editor/DataLayerEditor/Private/DataLayer/DataLayerEditorSubsystem.cpp
--- a/Engine/Source/Editor/DataLayerEditor/Private/DataLayer/DataLayerEditorSubsystem.cpp
+++ b/Engine/Source/Editor/DataLayerEditor/Private/DataLayer/DataLayerEditorSubsystem.cpp
@@ -1112,11 +1112,14 @@ bool UDataLayerEditorSubsystem::PromptDataLayerAssetSelection(FAssetData& OutAss
 	FOpenAssetDialogConfig Config;
 	Config.bAllowMultipleSelection = false;
 	Config.AssetClassNames.Add(UDataLayerAsset::StaticClass()->GetFName());
+	Config.DefaultPath = PickDataLayerDialogPath;
+	Config.DialogTitleOverride = LOCTEXT("PickDataLayerAssetDialogTitle", "Pick A Data Layer Asset");
 
 	TArray<FAssetData> Assets = ContentBrowserSingleton.CreateModalOpenAssetDialog(Config);
 	if (Assets.Num() == 1)
 	{
 		OutAsset = Assets[0];
+		OutAsset.PackagePath.ToString(PickDataLayerDialogPath);
 		return true;
 	}
 
diff --git a/Engine/Source/Editor/DataLayerEditor/Public/DataLayer/DataLayerEditorSubsystem.h b/Engine/Source/Editor/DataLayerEditor/Public/DataLayer/DataLayerEditorSubsystem.h
--- a/Engine/Source/Editor/DataLayerEditor/Public/DataLayer/DataLayerEditorSubsystem.h
+++ b/Engine/Source/Editor/DataLayerEditor/Public/DataLayer/DataLayerEditorSubsystem.h
@@ -662,5 +662,8 @@ private:
 	/** Delegate used to notify changes to ActorEditorContextSubsystem */
 	FOnActorEditorContextClientChanged ActorEditorContextClientChanged;
 
+	/** The path at which the "Pick A Data Layer Asset" will be opened*/
+	mutable FString PickDataLayerDialogPath;
+
 	friend class FDataLayersBroadcast;
 };
