Fix non-unity

diff --git a/Engine/Source/Editor/DataLayerEditor/Private/DataLayer/DataLayerOutlinerHasErrorColumn.cpp b/Engine/Source/Editor/DataLayerEditor/Private/DataLayer/DataLayerOutlinerHasErrorColumn.cpp
--- a/Engine/Source/Editor/DataLayerEditor/Private/DataLayer/DataLayerOutlinerHasErrorColumn.cpp
+++ b/Engine/Source/Editor/DataLayerEditor/Private/DataLayer/DataLayerOutlinerHasErrorColumn.cpp
@@ -2,6 +2,7 @@
 
 #include "DataLayerOutlinerHasErrorColumn.h"
 
+#include "Algo/Transform.h"
 #include "DataLayer/DataLayerEditorSubsystem.h"
 #include "DataLayerTreeItem.h"
 #include "Widgets/Images/SImage.h"
