Added a FDataLayerCreationParameters struct to supply data layer instance creation with more parameters
CreateDataLayer is now deprecated and CreateDataLayerInstance should be used
Move Data Layer Asset selection code from DataLayerEditorSubsystem to DataLayerMode

diff --git a/Engine/Source/Editor/DataLayerEditor/Private/DataLayer/DataLayerEditorSubsystem.cpp b/Engine/Source/Editor/DataLayerEditor/Private/DataLayer/DataLayerEditorSubsystem.cpp
--- a/Engine/Source/Editor/DataLayerEditor/Private/DataLayer/DataLayerEditorSubsystem.cpp
+++ b/Engine/Source/Editor/DataLayerEditor/Private/DataLayer/DataLayerEditorSubsystem.cpp
@@ -25,8 +25,6 @@
 #include "EngineUtils.h"
 #include "Editor.h"
 #include "Subsystems/ActorEditorContextSubsystem.h"
-#include "IContentBrowserSingleton.h"
-#include "ContentBrowserModule.h"
 #include "ProfilingDebugging/ScopedTimers.h"
 
 #define LOCTEXT_NAMESPACE "DataLayer"
@@ -1030,39 +1028,14 @@ void UDataLayerEditorSubsystem::AddAllDataLayersTo(TArray<TWeakObjectPtr<UDataLa
 	}
 }
 
-UDataLayerInstance* UDataLayerEditorSubsystem::CreateDataLayer(UDataLayerInstance* ParentDataLayer)
+UDataLayerInstance* UDataLayerEditorSubsystem::CreateDataLayerInstance(const FDataLayerCreationParameters& Parameters)
 {
 	UDataLayerInstance* NewDataLayer = nullptr;
 
-
-	AWorldDataLayers* WorldDataLayers = GetWorldDataLayers(/*bCreateIfNotFound*/true);
-	if (!WorldDataLayers->HasDeprecatedDataLayers())
+	AWorldDataLayers* WorldDataLayers = Parameters.WorlDataLayers != nullptr ? Parameters.WorlDataLayers.Get() : GetWorld()->GetWorldDataLayers();
+	if (!HasDeprecatedDataLayers())
 	{
-		bool bContinuePrompt = true;
-		while (bContinuePrompt)
-		{
-			FAssetData AssetData;
-			bContinuePrompt = PromptDataLayerAssetSelection(AssetData);
-			if (bContinuePrompt)
-			{
-				UDataLayerAsset* DataLayerAsset = CastChecked<UDataLayerAsset>(AssetData.GetAsset());
-				if (UDataLayerInstance* DataLayerInstance = GetDataLayerInstance(DataLayerAsset))
-				{
-					FText WarningMessage(LOCTEXT("CreateDataLayerAlreadyExist", "A Data Layer of that asset already exists. Pick another Data Layer?"));
-					bContinuePrompt = FMessageDialog::Open(EAppMsgType::YesNo, WarningMessage) == EAppReturnType::Yes;
-				}
-				else if (ParentDataLayer != nullptr && ParentDataLayer->GetType() != DataLayerAsset->GetType())
-				{
-					FText WarningMessage(LOCTEXT("CreateDataLayerDiffType", "The parent data layer is of a different type. Pick another Data Layer?"));
-					bContinuePrompt = FMessageDialog::Open(EAppMsgType::YesNo, WarningMessage) == EAppReturnType::Yes;
-				}
-				else
-				{
-					NewDataLayer = WorldDataLayers->CreateDataLayer<UDataLayerInstanceWithAsset>(DataLayerAsset);
-					bContinuePrompt = false;
-				}
-			}
-		}
+		NewDataLayer = WorldDataLayers->CreateDataLayer<UDataLayerInstanceWithAsset>(Parameters.DataLayerAsset);
 	}
 	else
 	{
@@ -1072,7 +1045,7 @@ UDataLayerInstance* UDataLayerEditorSubsystem::CreateDataLayer(UDataLayerInstanc
 	if (NewDataLayer != nullptr)
 	{
 		BroadcastDataLayerChanged(EDataLayerAction::Add, NewDataLayer, NAME_None);
-		SetParentDataLayer(NewDataLayer, ParentDataLayer);
+		SetParentDataLayer(NewDataLayer, Parameters.ParentDataLayer);
 	}
 	
 	return NewDataLayer;
@@ -1133,28 +1106,6 @@ void UDataLayerEditorSubsystem::RebuildSelectedDataLayersFromEditorSelection()
 	}
 }
 
-
-bool UDataLayerEditorSubsystem::PromptDataLayerAssetSelection(FAssetData& OutAsset) const
-{
-	IContentBrowserSingleton& ContentBrowserSingleton = FModuleManager::LoadModuleChecked<FContentBrowserModule>("ContentBrowser").Get();
-
-	FOpenAssetDialogConfig Config;
-	Config.bAllowMultipleSelection = false;
-	Config.AssetClassNames.Add(UDataLayerAsset::StaticClass()->GetFName());
-	Config.DefaultPath = PickDataLayerDialogPath;
-	Config.DialogTitleOverride = LOCTEXT("PickDataLayerAssetDialogTitle", "Pick A Data Layer Asset");
-
-	TArray<FAssetData> Assets = ContentBrowserSingleton.CreateModalOpenAssetDialog(Config);
-	if (Assets.Num() == 1)
-	{
-		OutAsset = Assets[0];
-		OutAsset.PackagePath.ToString(PickDataLayerDialogPath);
-		return true;
-	}
-
-	return false;
-}
-
 //~ Begin Deprecated
 
 PRAGMA_DISABLE_DEPRECATION_WARNINGS
diff --git a/Engine/Source/Editor/DataLayerEditor/Private/DataLayer/DataLayerMode.cpp b/Engine/Source/Editor/DataLayerEditor/Private/DataLayer/DataLayerMode.cpp
--- a/Engine/Source/Editor/DataLayerEditor/Private/DataLayer/DataLayerMode.cpp
+++ b/Engine/Source/Editor/DataLayerEditor/Private/DataLayer/DataLayerMode.cpp
@@ -3,9 +3,12 @@
 #include "DataLayerMode.h"
 #include "SSceneOutliner.h"
 #include "SDataLayerBrowser.h"
+#include "IContentBrowserSingleton.h"
+#include "ContentBrowserModule.h"
 #include "DataLayer/DataLayerEditorSubsystem.h"
 #include "DataLayer/DataLayerOutlinerDeleteButtonColumn.h"
 #include "WorldPartition/DataLayer/WorldDataLayers.h"
+#include "WorldPartition/DataLayer/DataLayerAsset.h"
 #include "DataLayer/SDataLayerOutliner.h"
 #include "DataLayerHierarchy.h"
 #include "DataLayerActorTreeItem.h"
@@ -928,11 +931,17 @@ void FDataLayerMode::RegisterContextMenu()
 					const FScopedDataLayerTransaction Transaction(LOCTEXT("CreateNewDataLayer", "Create New Data Layer"), RepresentingWorld.Get());
 					SelectedDataLayersSet.Empty();
 					SelectedDataLayerActors.Empty();
-					if (UDataLayerInstance* NewDataLayerInstance = DataLayerEditorSubsystem->CreateDataLayer(ParentDataLayer))
+					if (UDataLayerAsset* DataLayerAsset = PromptDataLayerAssetSelection())
 					{
-						SelectedDataLayersSet.Add(NewDataLayerInstance);
-						// Select it and open a rename when it gets refreshed
-						SceneOutliner->OnItemAdded(NewDataLayerInstance, SceneOutliner::ENewItemAction::Select | SceneOutliner::ENewItemAction::Rename);
+						FDataLayerCreationParameters CreationParams;
+						CreationParams.DataLayerAsset = DataLayerAsset;
+						CreationParams.ParentDataLayer = ParentDataLayer;
+						if (UDataLayerInstance* NewDataLayerInstance = DataLayerEditorSubsystem->CreateDataLayerInstance(CreationParams))
+						{
+							SelectedDataLayersSet.Add(NewDataLayerInstance);
+							// Select it and open a rename when it gets refreshed
+							SceneOutliner->OnItemAdded(NewDataLayerInstance, SceneOutliner::ENewItemAction::Select | SceneOutliner::ENewItemAction::Rename);
+						}
 					}
 				};
 
@@ -949,10 +958,15 @@ void FDataLayerMode::RegisterContextMenu()
 				Section.AddMenuEntry("AddSelectedActorsToNewDataLayer", LOCTEXT("AddSelectedActorsToNewDataLayer", "Add Selected Actors to New Data Layer"), FText(), FSlateIcon(),
 					FUIAction(
 						FExecuteAction::CreateLambda([this]() {
-							const FScopedDataLayerTransaction Transaction(LOCTEXT("AddSelectedActorsToNewDataLayer", "Add Selected Actors to New Data Layer"), RepresentingWorld.Get());
-							if (UDataLayerInstance* NewDataLayer = DataLayerEditorSubsystem->CreateDataLayer())
+							if (UDataLayerAsset* DataLayerAsset = PromptDataLayerAssetSelection())
 							{
-								DataLayerEditorSubsystem->AddSelectedActorsToDataLayer(NewDataLayer);
+								const FScopedDataLayerTransaction Transaction(LOCTEXT("AddSelectedActorsToNewDataLayer", "Add Selected Actors to New Data Layer"), RepresentingWorld.Get());
+								FDataLayerCreationParameters CreationParams;
+								CreationParams.DataLayerAsset = DataLayerAsset;
+								if (UDataLayerInstance* NewDataLayerInstance = DataLayerEditorSubsystem->CreateDataLayerInstance(CreationParams))
+								{
+									DataLayerEditorSubsystem->AddSelectedActorsToDataLayer(NewDataLayerInstance);
+								}
 							}}),
 						FCanExecuteAction::CreateLambda([] { return GEditor->GetSelectedActorCount() > 0; })
 					));
@@ -1520,6 +1534,26 @@ TSharedRef<FSceneOutlinerFilter> FDataLayerMode::CreateShowOnlySelectedActorsFil
 	return MakeShareable(new FDataLayerActorFilter(FDataLayerActorTreeItem::FFilterPredicate::CreateStatic(IsActorSelected), FSceneOutlinerFilter::EDefaultBehaviour::Pass, FDataLayerActorTreeItem::FFilterPredicate::CreateStatic(IsActorSelected)));
 }
 
+UDataLayerAsset* FDataLayerMode::PromptDataLayerAssetSelection()
+{
+	IContentBrowserSingleton& ContentBrowserSingleton = FModuleManager::LoadModuleChecked<FContentBrowserModule>("ContentBrowser").Get();
+
+	FOpenAssetDialogConfig Config;
+	Config.bAllowMultipleSelection = false;
+	Config.AssetClassNames.Add(UDataLayerAsset::StaticClass()->GetFName());
+	Config.DefaultPath = PickDataLayerDialogPath;
+	Config.DialogTitleOverride = LOCTEXT("PickDataLayerAssetDialogTitle", "Pick a Data Layer Asset");
+
+	TArray<FAssetData> Assets = ContentBrowserSingleton.CreateModalOpenAssetDialog(Config);
+	if (Assets.Num() == 1)
+	{
+		Assets[0].PackagePath.ToString(PickDataLayerDialogPath);
+		return CastChecked<UDataLayerAsset>(Assets[0].GetAsset());
+	}
+
+	return nullptr;
+}
+
 void FDataLayerMode::SynchronizeSelection()
 {
 	if (!bShowOnlySelectedActors && !bHighlightSelectedDataLayers)
diff --git a/Engine/Source/Editor/DataLayerEditor/Private/DataLayer/DataLayerMode.h b/Engine/Source/Editor/DataLayerEditor/Private/DataLayer/DataLayerMode.h
--- a/Engine/Source/Editor/DataLayerEditor/Private/DataLayer/DataLayerMode.h
+++ b/Engine/Source/Editor/DataLayerEditor/Private/DataLayer/DataLayerMode.h
@@ -111,6 +111,8 @@ private:
 	void RefreshSelection();
 	FSceneOutlinerDragValidationInfo ValidateDrop(const ISceneOutlinerTreeItem& DropTarget, bool bMoveOperation = false) const;
 
+	UDataLayerAsset* PromptDataLayerAssetSelection();
+
 	/** Filter factories */
 	static TSharedRef<FSceneOutlinerFilter> CreateShowOnlySelectedActorsFilter();
 	static TSharedRef<FSceneOutlinerFilter> CreateHideEditorDataLayersFilter();
@@ -132,6 +134,8 @@ private:
 	uint32 FilteredDataLayerCount = 0;
 	/** List of datalayers which passed the regular filters and may or may not have passed the search filter */
 	TSet<TWeakObjectPtr<UDataLayerInstance>> ApplicableDataLayers;
+	/** The path at which the "Pick A Data Layer Asset" will be opened*/
+	mutable FString PickDataLayerDialogPath;
 
 	TSet<TWeakObjectPtr<const UDataLayerInstance>> SelectedDataLayersSet;
 	typedef TPair<TWeakObjectPtr<const UDataLayerInstance>, TWeakObjectPtr<const AActor>> FSelectedDataLayerActor;
diff --git a/Engine/Source/Editor/DataLayerEditor/Public/DataLayer/DataLayerEditorSubsystem.h b/Engine/Source/Editor/DataLayerEditor/Public/DataLayer/DataLayerEditorSubsystem.h
--- a/Engine/Source/Editor/DataLayerEditor/Public/DataLayer/DataLayerEditorSubsystem.h
+++ b/Engine/Source/Editor/DataLayerEditor/Public/DataLayer/DataLayerEditorSubsystem.h
@@ -18,6 +18,24 @@ class ULevel;
 class UWorld;
 template<typename TItemType> class IFilter;
 
+USTRUCT(BlueprintType)
+struct FDataLayerCreationParameters
+{
+	GENERATED_USTRUCT_BODY()
+
+	// Required. Will assign the asset to the created instance.
+	UPROPERTY(EditAnywhere, BlueprintReadWrite, Category = "Data Layer")
+	TObjectPtr<UDataLayerAsset> DataLayerAsset;
+
+	// Optional. Will create the Data Layer at root level is unset.
+	UPROPERTY()
+	TObjectPtr<UDataLayerInstance> ParentDataLayer;
+
+	// Optional. Will default at the level WorldDataLayers if unset.
+	UPROPERTY()
+	TObjectPtr<AWorldDataLayers> WorlDataLayers;
+};
+
 UCLASS()
 class DATALAYEREDITOR_API UDataLayerEditorSubsystem final : public UEditorSubsystem, public IActorEditorContextClient
 {
@@ -485,11 +503,11 @@ public:
 	/**
 	 * Creates a UDataLayerInstance Object
 	 *
-	 * @param	ParentDataLayer	An optional data layer that will parent the newly created data layer
+	 * @param	Parameters The Data Layer Instance creation parameters
 	 * @return	The newly created UDataLayerInstance Object
 	 */
 	UFUNCTION(BlueprintCallable, Category = DataLayers)
-	UDataLayerInstance* CreateDataLayer(UDataLayerInstance* ParentDataLayer = nullptr);
+	UDataLayerInstance* CreateDataLayerInstance(const FDataLayerCreationParameters& Parameters);
 
 	/**
 	 * Sets a Parent DataLayer for a specified DataLayer
@@ -587,6 +605,10 @@ public:
 	UFUNCTION(BlueprintCallable, Category = DataLayers)
 	bool RenameDataLayer(UDataLayerInstance* DataLayer, const FName& NewDataLayerLabel);
 
+	UE_DEPRECATED(5.1, "Use CreateDataLayerInstance with FDataLayerCreationParameters instead")
+	UFUNCTION(BlueprintCallable, Category = DataLayers)
+	UDataLayerInstance* CreateDataLayer(UDataLayerInstance* ParentDataLayer = nullptr) { return nullptr; }
+
 	//~ End Deprecated
 
 private:
@@ -645,8 +667,6 @@ private:
 	void OnSelectionChanged();
 	void RebuildSelectedDataLayersFromEditorSelection();
 
-	bool PromptDataLayerAssetSelection(FAssetData& OutAsset) const;
-
 	/** Contains Data Layers that contain actors that are part of the editor selection */
 	TSet<TWeakObjectPtr<const UDataLayerInstance>> SelectedDataLayersFromEditorSelection;
 
@@ -662,9 +682,6 @@ private:
 	/** Delegate used to notify changes to ActorEditorContextSubsystem */
 	FOnActorEditorContextClientChanged ActorEditorContextClientChanged;
 
-	/** The path at which the "Pick A Data Layer Asset" will be opened*/
-	mutable FString PickDataLayerDialogPath;
-
 	static void OnActorDescContainerInitialized(class UActorDescContainer* InActorDescContainer);
 
 	friend class FDataLayersBroadcast;
