Fix uninitialized properties warning message

diff --git a/Engine/Source/Editor/DataLayerEditor/Private/DataLayer/DataLayerEditorSubsystem.cpp b/Engine/Source/Editor/DataLayerEditor/Private/DataLayer/DataLayerEditorSubsystem.cpp
--- a/Engine/Source/Editor/DataLayerEditor/Private/DataLayer/DataLayerEditorSubsystem.cpp
+++ b/Engine/Source/Editor/DataLayerEditor/Private/DataLayer/DataLayerEditorSubsystem.cpp
@@ -31,6 +31,14 @@
 
 DEFINE_LOG_CATEGORY_STATIC(LogDataLayerEditorSubsystem, All, All);
 
+FDataLayerCreationParameters::FDataLayerCreationParameters()
+	:DataLayerAsset(nullptr),
+	ParentDataLayer(nullptr),
+	WorlDataLayers(nullptr)
+{
+
+}
+
 //////////////////////////////////////////////////////////////////////////
 // FDataLayersBroadcast
 
diff --git a/Engine/Source/Editor/DataLayerEditor/Public/DataLayer/DataLayerEditorSubsystem.h b/Engine/Source/Editor/DataLayerEditor/Public/DataLayer/DataLayerEditorSubsystem.h
--- a/Engine/Source/Editor/DataLayerEditor/Public/DataLayer/DataLayerEditorSubsystem.h
+++ b/Engine/Source/Editor/DataLayerEditor/Public/DataLayer/DataLayerEditorSubsystem.h
@@ -23,6 +23,8 @@ struct FDataLayerCreationParameters
 {
 	GENERATED_USTRUCT_BODY()
 
+	FDataLayerCreationParameters();
+
 	// Required. Will assign the asset to the created instance.
 	UPROPERTY(EditAnywhere, BlueprintReadWrite, Category = "Data Layer")
 	TObjectPtr<UDataLayerAsset> DataLayerAsset;
