-Unregister context menu when FDataLayerMode is destroyed since the context menu uses lambdas that capture the FDataLayerMode instance and doesn't get recreated unless unregistered.

diff --git a/Engine/Source/Editor/DataLayerEditor/Private/DataLayer/DataLayerMode.cpp b/Engine/Source/Editor/DataLayerEditor/Private/DataLayer/DataLayerMode.cpp
--- a/Engine/Source/Editor/DataLayerEditor/Private/DataLayer/DataLayerMode.cpp
+++ b/Engine/Source/Editor/DataLayerEditor/Private/DataLayer/DataLayerMode.cpp
@@ -145,6 +145,8 @@ FDataLayerMode::FDataLayerMode(const FDataLayerModeParams& Params)
 
 FDataLayerMode::~FDataLayerMode()
 {
+	UnregisterContextMenu();
+
 	USelection::SelectionChangedEvent.RemoveAll(this);
 	USelection::SelectObjectEvent.RemoveAll(this);
 }
@@ -1200,6 +1202,13 @@ void FDataLayerMode::RegisterContextMenu()
 	}
 }
 
+void FDataLayerMode::UnregisterContextMenu()
+{
+	UToolMenus* ToolMenus = UToolMenus::Get();
+	ToolMenus->RemoveMenu(DefaultContextBaseMenuName);
+	ToolMenus->RemoveMenu(DefaultContextMenuName);
+}
+
 TSharedPtr<SWidget> FDataLayerMode::CreateContextMenu()
 {
 	RegisterContextMenu();
diff --git a/Engine/Source/Editor/DataLayerEditor/Private/DataLayer/DataLayerMode.h b/Engine/Source/Editor/DataLayerEditor/Private/DataLayer/DataLayerMode.h
--- a/Engine/Source/Editor/DataLayerEditor/Private/DataLayer/DataLayerMode.h
+++ b/Engine/Source/Editor/DataLayerEditor/Private/DataLayer/DataLayerMode.h
@@ -96,6 +96,7 @@ protected:
 private:
 	/* Private Helpers */
 	void RegisterContextMenu();
+	void UnregisterContextMenu();
 	void ChooseRepresentingWorld();
 	void OnSelectWorld(TWeakObjectPtr<UWorld> World);
 	bool IsWorldChecked(TWeakObjectPtr<UWorld> World) const;
