Changed SelectedDataLayersFromEditorSelection to be rebuilt only when necessary.

diff --git a/Engine/Source/Editor/DataLayerEditor/Private/DataLayer/DataLayerEditorSubsystem.cpp b/Engine/Source/Editor/DataLayerEditor/Private/DataLayer/DataLayerEditorSubsystem.cpp
--- a/Engine/Source/Editor/DataLayerEditor/Private/DataLayer/DataLayerEditorSubsystem.cpp
+++ b/Engine/Source/Editor/DataLayerEditor/Private/DataLayer/DataLayerEditorSubsystem.cpp
@@ -140,6 +140,10 @@ void FDataLayersBroadcast::OnObjectPostEditChange(UObject* Object, FPropertyChan
 //		  Either a proxy that redirects calls to the proper EditorSubsystem will be used or user code will change to trigger delegate broadcast instead of directly accessing the subsystem (see calls to InitializeNewActorDataLayers everywhere as an example).
 //
 
+UDataLayerEditorSubsystem::UDataLayerEditorSubsystem()
+: bRebuildSelectedDataLayersFromEditorSelection(false)
+{}
+
 UDataLayerEditorSubsystem* UDataLayerEditorSubsystem::Get()
 {
 	return GEditor ? GEditor->GetEditorSubsystem<UDataLayerEditorSubsystem>() : nullptr;
@@ -1072,35 +1076,40 @@ void UDataLayerEditorSubsystem::DeleteDataLayer(UDataLayerInstance* DataLayerToD
 
 void UDataLayerEditorSubsystem::BroadcastActorDataLayersChanged(const TWeakObjectPtr<AActor>& ChangedActor)
 {
-	RebuildSelectedDataLayersFromEditorSelection();
+	bRebuildSelectedDataLayersFromEditorSelection = true;
 	ActorDataLayersChanged.Broadcast(ChangedActor);
 }
 
 void UDataLayerEditorSubsystem::BroadcastDataLayerChanged(const EDataLayerAction Action, const TWeakObjectPtr<const UDataLayerInstance>& ChangedDataLayer, const FName& ChangedProperty)
 {
-	RebuildSelectedDataLayersFromEditorSelection();
+	bRebuildSelectedDataLayersFromEditorSelection = true;
 	DataLayerChanged.Broadcast(Action, ChangedDataLayer, ChangedProperty);
 	ActorEditorContextClientChanged.Broadcast(this);
 }
 
 void UDataLayerEditorSubsystem::OnSelectionChanged()
 {
-	RebuildSelectedDataLayersFromEditorSelection();
+	bRebuildSelectedDataLayersFromEditorSelection = true;
 }
 
-void UDataLayerEditorSubsystem::RebuildSelectedDataLayersFromEditorSelection()
+const TSet<TWeakObjectPtr<const UDataLayerInstance>>& UDataLayerEditorSubsystem::GetSelectedDataLayersFromEditorSelection() const
 {
-	SelectedDataLayersFromEditorSelection.Reset();
-
-	TArray<AActor*> Actors;
-	GEditor->GetSelectedActors()->GetSelectedObjects<AActor>(Actors);
-	for (const AActor* Actor : Actors)
+	if (bRebuildSelectedDataLayersFromEditorSelection)
 	{
-		for (const UDataLayerInstance* DataLayerInstance : Actor->GetDataLayerInstances())
+		bRebuildSelectedDataLayersFromEditorSelection = false;
+
+		SelectedDataLayersFromEditorSelection.Reset();
+		TArray<AActor*> Actors;
+		GEditor->GetSelectedActors()->GetSelectedObjects<AActor>(Actors);
+		for (const AActor* Actor : Actors)
 		{
-			SelectedDataLayersFromEditorSelection.Add(DataLayerInstance);
+			for (const UDataLayerInstance* DataLayerInstance : Actor->GetDataLayerInstances())
+			{
+				SelectedDataLayersFromEditorSelection.Add(DataLayerInstance);
+			}
 		}
 	}
+	return SelectedDataLayersFromEditorSelection;
 }
 
 //~ Begin Deprecated
diff --git a/Engine/Source/Editor/DataLayerEditor/Public/DataLayer/DataLayerEditorSubsystem.h b/Engine/Source/Editor/DataLayerEditor/Public/DataLayer/DataLayerEditorSubsystem.h
--- a/Engine/Source/Editor/DataLayerEditor/Public/DataLayer/DataLayerEditorSubsystem.h
+++ b/Engine/Source/Editor/DataLayerEditor/Public/DataLayer/DataLayerEditorSubsystem.h
@@ -44,6 +44,8 @@ class DATALAYEREDITOR_API UDataLayerEditorSubsystem final : public UEditorSubsys
 	GENERATED_BODY()
 
 public:
+	UDataLayerEditorSubsystem();
+
 	static UDataLayerEditorSubsystem* Get();
 
 	typedef IFilter<const TWeakObjectPtr<AActor>&> FActorFilter;
@@ -546,7 +548,7 @@ public:
 	/*
 	 * Returns whether the DataLayer contains one of more actors that are part of the editor selection.
 	 */
-	bool DoesDataLayerContainSelectedActors(const UDataLayerInstance* DataLayer) const { return SelectedDataLayersFromEditorSelection.Contains(DataLayer); }
+	bool DoesDataLayerContainSelectedActors(const UDataLayerInstance* DataLayer) const { return GetSelectedDataLayersFromEditorSelection().Contains(DataLayer); }
 
 	/**
 	* Whether there are any deprecated DataLayerInstance found
@@ -660,9 +662,13 @@ private:
 	void BroadcastDataLayerChanged(const EDataLayerAction Action, const TWeakObjectPtr<const UDataLayerInstance>& ChangedDataLayer, const FName& ChangedProperty);
 	void OnSelectionChanged();
 	void RebuildSelectedDataLayersFromEditorSelection();
+	const TSet<TWeakObjectPtr<const UDataLayerInstance>>& GetSelectedDataLayersFromEditorSelection() const;
 
 	/** Contains Data Layers that contain actors that are part of the editor selection */
-	TSet<TWeakObjectPtr<const UDataLayerInstance>> SelectedDataLayersFromEditorSelection;
+	mutable TSet<TWeakObjectPtr<const UDataLayerInstance>> SelectedDataLayersFromEditorSelection;
+
+	/** Internal flag to know if SelectedDataLayersFromEditorSelection needs to be rebuilt. */
+	mutable bool bRebuildSelectedDataLayersFromEditorSelection;
 
 	/** Fires whenever one or more DataLayer changes */
 	FOnDataLayerChanged DataLayerChanged;
