Data Layers: added engine bridge for editor callbacks.

#rnx

diff --git a/Engine/Source/Editor/DataLayerEditor/Private/DataLayer/DataLayerEditorSubsystem.cpp b/Engine/Source/Editor/DataLayerEditor/Private/DataLayer/DataLayerEditorSubsystem.cpp
--- a/Engine/Source/Editor/DataLayerEditor/Private/DataLayer/DataLayerEditorSubsystem.cpp
+++ b/Engine/Source/Editor/DataLayerEditor/Private/DataLayer/DataLayerEditorSubsystem.cpp
@@ -164,6 +164,9 @@ void UDataLayerEditorSubsystem::Initialize(FSubsystemCollectionBase& Collection)
 	}
 
 	UActorEditorContextSubsystem::Get()->RegisterClient(this);
+
+	// Register the engine broadcast bridge
+	OnActorDataLayersEditorLoadingStateChangedEngineBridgeHandle = DataLayerEditorLoadingStateChanged.AddStatic(&FDataLayersEditorBroadcast::StaticOnActorDataLayersEditorLoadingStateChanged);
 }
 
 void UDataLayerEditorSubsystem::Deinitialize()
@@ -173,6 +176,9 @@ void UDataLayerEditorSubsystem::Deinitialize()
 	Super::Deinitialize();
 
 	DataLayersBroadcast->Deinitialize();
+
+	// Unregister the engine broadcast bridge
+	DataLayerEditorLoadingStateChanged.Remove(OnActorDataLayersEditorLoadingStateChangedEngineBridgeHandle);
 }
 
 void UDataLayerEditorSubsystem::OnActorDescContainerInitialized(UActorDescContainer* InActorDescContainer)
@@ -284,19 +290,6 @@ void UDataLayerEditorSubsystem::RemoveFromActorEditorContext(UDataLayerInstance*
 	}
 }
 
-bool UDataLayerEditorSubsystem::RefreshWorldPartitionEditorCells(bool bIsFromUserChange)
-{
-	if (UWorldPartition* WorldPartition = GetWorld() ? GetWorld()->GetWorldPartition() : nullptr)
-	{
-		if (!WorldPartition->RefreshLoadedEditorCells(bIsFromUserChange))
-		{
-			return false;
-		}
-		UpdateDataLayerEditorPerProjectUserSettings();
-	}
-	return true;
-}
-
 void UDataLayerEditorSubsystem::UpdateDataLayerEditorPerProjectUserSettings()
 {
 	if (AWorldDataLayers* WorldDataLayers = GetWorldDataLayers())
@@ -371,7 +364,7 @@ bool UDataLayerEditorSubsystem::SetParentDataLayer(UDataLayerInstance* DataLayer
 		UpdateAllActorsVisibility(true, true);
 		if (bIsLoaded != DataLayer->IsEffectiveLoadedInEditor())
 		{
-			RefreshWorldPartitionEditorCells(true);
+			BroadcastDataLayerEditorLoadingStateChanged(true);
 		}
 		return true;
 	}
@@ -945,8 +938,11 @@ bool UDataLayerEditorSubsystem::SetDataLayerIsLoadedInEditorInternal(UDataLayerI
 
 bool UDataLayerEditorSubsystem::SetDataLayerIsLoadedInEditor(UDataLayerInstance* DataLayer, const bool bIsLoadedInEditor, const bool bIsFromUserChange)
 {
-	bool bRefreshNeeded = SetDataLayerIsLoadedInEditorInternal(DataLayer, bIsLoadedInEditor, bIsFromUserChange);
-	return bRefreshNeeded ? RefreshWorldPartitionEditorCells(bIsFromUserChange) : true;
+	if (SetDataLayerIsLoadedInEditorInternal(DataLayer, bIsLoadedInEditor, bIsFromUserChange))
+	{
+		BroadcastDataLayerEditorLoadingStateChanged(bIsFromUserChange);
+	}
+	return true;
 }
 
 bool UDataLayerEditorSubsystem::SetDataLayersIsLoadedInEditor(const TArray<UDataLayerInstance*>& DataLayers, const bool bIsLoadedInEditor, const bool bIsFromUserChange)
@@ -956,7 +952,13 @@ bool UDataLayerEditorSubsystem::SetDataLayersIsLoadedInEditor(const TArray<UData
 	{
 		bRefreshNeeded |= SetDataLayerIsLoadedInEditorInternal(DataLayer, bIsLoadedInEditor, bIsFromUserChange);
 	}
-	return bRefreshNeeded ? RefreshWorldPartitionEditorCells(bIsFromUserChange) : true;
+	
+	if (bRefreshNeeded)
+	{
+		BroadcastDataLayerEditorLoadingStateChanged(bIsFromUserChange);
+	}
+
+	return true;
 }
 
 bool UDataLayerEditorSubsystem::ToggleDataLayerIsLoadedInEditor(UDataLayerInstance* DataLayer, const bool bIsFromUserChange)
@@ -972,7 +974,13 @@ bool UDataLayerEditorSubsystem::ToggleDataLayersIsLoadedInEditor(const TArray<UD
 	{
 		bRefreshNeeded |= SetDataLayerIsLoadedInEditorInternal(DataLayer, !DataLayer->IsLoadedInEditor(), bIsFromUserChange);
 	}
-	return bRefreshNeeded ? RefreshWorldPartitionEditorCells(bIsFromUserChange) : true;
+	
+	if (bRefreshNeeded)
+	{
+		BroadcastDataLayerEditorLoadingStateChanged(bIsFromUserChange);
+	}
+
+	return true;
 }
 
 bool UDataLayerEditorSubsystem::ResetUserSettings()
@@ -986,7 +994,13 @@ bool UDataLayerEditorSubsystem::ResetUserSettings()
 			return true;
 		});
 	}
-	return bRefreshNeeded ? RefreshWorldPartitionEditorCells(true) : true;
+	
+	if (bRefreshNeeded)
+	{
+		BroadcastDataLayerEditorLoadingStateChanged(true);
+	}
+
+	return true;
 }
 
 bool UDataLayerEditorSubsystem::HasDeprecatedDataLayers() const
@@ -1087,6 +1101,12 @@ void UDataLayerEditorSubsystem::BroadcastDataLayerChanged(const EDataLayerAction
 	ActorEditorContextClientChanged.Broadcast(this);
 }
 
+void UDataLayerEditorSubsystem::BroadcastDataLayerEditorLoadingStateChanged(bool bIsFromUserChange)
+{
+	DataLayerEditorLoadingStateChanged.Broadcast(bIsFromUserChange);
+	UpdateDataLayerEditorPerProjectUserSettings();
+}
+
 void UDataLayerEditorSubsystem::OnSelectionChanged()
 {
 	bRebuildSelectedDataLayersFromEditorSelection = true;
diff --git a/Engine/Source/Editor/DataLayerEditor/Public/DataLayer/DataLayerEditorSubsystem.h b/Engine/Source/Editor/DataLayerEditor/Public/DataLayer/DataLayerEditorSubsystem.h
--- a/Engine/Source/Editor/DataLayerEditor/Public/DataLayer/DataLayerEditorSubsystem.h
+++ b/Engine/Source/Editor/DataLayerEditor/Public/DataLayer/DataLayerEditorSubsystem.h
@@ -91,7 +91,11 @@ public:
 	/** Broadcasts whenever one or more Actors changed UDataLayerInstances*/
 	DECLARE_EVENT_OneParam(UDataLayerEditorSubsystem, FOnActorDataLayersChanged, const TWeakObjectPtr<AActor>& /*ChangedActor*/);
 	FOnActorDataLayersChanged& OnActorDataLayersChanged() { return ActorDataLayersChanged; }
-	
+
+	/** Broadcasts whenever one or more DataLayers editor loading state changed */
+	DECLARE_EVENT_OneParam(UDataLayerEditorSubsystem, FOnActorDataLayersEditorLoadingStateChanged, bool /*bIsFromUserChange*/);
+	FOnActorDataLayersEditorLoadingStateChanged& OnActorDataLayersEditorLoadingStateChanged() { return DataLayerEditorLoadingStateChanged; }
+
 	///////////////////////////////////////////////////////////////////////////////////////////////////////////////////////
 	// Operations on an individual actor.
 
@@ -655,11 +659,11 @@ private:
 	void PostUndoRedo();
 
 	bool SetDataLayerIsLoadedInEditorInternal(UDataLayerInstance* DataLayer, const bool bIsLoadedInEditor, const bool bIsFromUserChange);
-	bool RefreshWorldPartitionEditorCells(bool bIsFromUserChange);
 	void UpdateDataLayerEditorPerProjectUserSettings();
 
 	void BroadcastActorDataLayersChanged(const TWeakObjectPtr<AActor>& ChangedActor);
 	void BroadcastDataLayerChanged(const EDataLayerAction Action, const TWeakObjectPtr<const UDataLayerInstance>& ChangedDataLayer, const FName& ChangedProperty);
+	void BroadcastDataLayerEditorLoadingStateChanged(bool bIsFromUserChange);
 	void OnSelectionChanged();
 	void RebuildSelectedDataLayersFromEditorSelection();
 	const TSet<TWeakObjectPtr<const UDataLayerInstance>>& GetSelectedDataLayersFromEditorSelection() const;
@@ -676,6 +680,11 @@ private:
 	/**	Fires whenever one or more actor DataLayer changes */
 	FOnActorDataLayersChanged ActorDataLayersChanged;
 
+	/** Fires whenever one or more DataLayer editor loading state changed */
+	FOnActorDataLayersEditorLoadingStateChanged DataLayerEditorLoadingStateChanged;
+
+	FDelegateHandle OnActorDataLayersEditorLoadingStateChangedEngineBridgeHandle;
+
 	/** Auxiliary class that sets the callback functions for multiple delegates */
 	TSharedPtr<class FDataLayersBroadcast> DataLayersBroadcast;
 
diff --git a/Engine/Source/Editor/UnrealEd/Private/WorldPartition/WorldPartitionBuilder.cpp b/Engine/Source/Editor/UnrealEd/Private/WorldPartition/WorldPartitionBuilder.cpp
--- a/Engine/Source/Editor/UnrealEd/Private/WorldPartition/WorldPartitionBuilder.cpp
+++ b/Engine/Source/Editor/UnrealEd/Private/WorldPartition/WorldPartitionBuilder.cpp
@@ -12,6 +12,7 @@
 #include "WorldPartition/WorldPartition.h"
 #include "WorldPartition/WorldPartitionHelpers.h"
 #include "WorldPartition/WorldPartitionSubsystem.h"
+#include "WorldPartition/DataLayer/DataLayerSubsystem.h"
 #include "WorldPartition/DataLayer/WorldDataLayers.h"
 #include "WorldPartition/DataLayer/DataLayerInstance.h"
 #include "LevelInstance/LevelInstanceSubsystem.h"
@@ -188,7 +189,7 @@ bool UWorldPartitionBuilder::Run(UWorld* World, FPackageSourceControlHelper& Pac
 		if (bUpdateEditorCells)
 		{
 			UE_LOG(LogWorldPartitionBuilder, Display, TEXT("DataLayer load state changed refreshing editor cells"));
-			WorldPartition->RefreshLoadedEditorCells(false);
+			FDataLayersEditorBroadcast::StaticOnActorDataLayersEditorLoadingStateChanged(false);
 		}
 	}
 
diff --git a/Engine/Source/Runtime/Engine/Private/WorldPartition/DataLayer/DataLayerSubsystem.cpp b/Engine/Source/Runtime/Engine/Private/WorldPartition/DataLayer/DataLayerSubsystem.cpp
--- a/Engine/Source/Runtime/Engine/Private/WorldPartition/DataLayer/DataLayerSubsystem.cpp
+++ b/Engine/Source/Runtime/Engine/Private/WorldPartition/DataLayer/DataLayerSubsystem.cpp
@@ -38,6 +38,19 @@ static FAutoConsoleCommandWithOutputDevice GDumpDataLayersCmd(
 	})
 );
 
+#if WITH_EDITOR
+FDataLayersEditorBroadcast& FDataLayersEditorBroadcast::Get()
+{
+	static FDataLayersEditorBroadcast DataLayersEditorBroadcast;
+	return DataLayersEditorBroadcast;
+}
+
+void FDataLayersEditorBroadcast::StaticOnActorDataLayersEditorLoadingStateChanged(bool bIsFromUserChange)
+{
+	Get().DataLayerEditorLoadingStateChanged.Broadcast(bIsFromUserChange);
+}
+#endif
+
 UDataLayerSubsystem::UDataLayerSubsystem()
 {}
 
diff --git a/Engine/Source/Runtime/Engine/Private/WorldPartition/WorldPartition.cpp b/Engine/Source/Runtime/Engine/Private/WorldPartition/WorldPartition.cpp
--- a/Engine/Source/Runtime/Engine/Private/WorldPartition/WorldPartition.cpp
+++ b/Engine/Source/Runtime/Engine/Private/WorldPartition/WorldPartition.cpp
@@ -41,6 +41,7 @@
 #include "WorldPartition/WorldPartitionEditorPerProjectUserSettings.h"
 #include "WorldPartition/WorldPartitionMiniMap.h"
 #include "WorldPartition/WorldPartitionMiniMapHelper.h"
+#include "WorldPartition/DataLayer/DataLayerSubsystem.h"
 #include "WorldPartition/DataLayer/DataLayerInstance.h"
 #include "WorldPartition/DataLayer/WorldDataLayers.h"
 #include "WorldPartition/WorldPartitionActorDescViewProxy.h"
@@ -653,6 +654,25 @@ void UWorldPartition::OnPostBugItGoCalled(const FVector& Loc, const FRotator& Ro
 #endif
 }
 
+#if WITH_EDITOR
+void UWorldPartition::OnActorDataLayersEditorLoadingStateChanged(bool bFromUserOperation)
+{
+	auto GetCellsToRefresh = [this](TArray<UWorldPartitionEditorCell*>& CellsToRefresh)
+	{
+		EditorHash->ForEachCell([&CellsToRefresh](UWorldPartitionEditorCell* Cell)
+		{
+			if (Cell->IsLoaded())
+			{
+				CellsToRefresh.Add(Cell);
+			}
+		});
+		return !CellsToRefresh.IsEmpty();
+	};
+
+	UpdateEditorCells(GetCellsToRefresh, /*bIsCellShouldBeLoaded*/true, bFromUserOperation);
+}
+#endif
+
 void UWorldPartition::RegisterDelegates()
 {
 	check(World); 
@@ -670,6 +690,8 @@ void UWorldPartition::RegisterDelegates()
 			FCoreUObjectDelegates::PostReachabilityAnalysis.AddUObject(this, &UWorldPartition::OnGCPostReachabilityAnalysis);
 
 			GEditor->OnPostBugItGoCalled().AddUObject(this, &UWorldPartition::OnPostBugItGoCalled);
+
+			FDataLayersEditorBroadcast::Get().OnActorDataLayersEditorLoadingStateChanged().AddUObject(this, &UWorldPartition::OnActorDataLayersEditorLoadingStateChanged);
 		}
 #endif
 
@@ -704,6 +726,8 @@ void UWorldPartition::UnregisterDelegates()
 			}
 
 			GEditor->OnPostBugItGoCalled().RemoveAll(this);
+
+			FDataLayersEditorBroadcast::Get().OnActorDataLayersEditorLoadingStateChanged().RemoveAll(this);
 		}
 #endif
 
@@ -853,23 +877,6 @@ void UWorldPartition::LoadEditorCells(const TArray<FName>& CellNames, bool bIsFr
 	}
 }
 
-bool UWorldPartition::RefreshLoadedEditorCells(bool bIsFromUserChange)
-{
-	auto GetCellsToRefresh = [this](TArray<UWorldPartitionEditorCell*>& CellsToRefresh)
-	{
-		EditorHash->ForEachCell([&CellsToRefresh](UWorldPartitionEditorCell* Cell)
-		{
-			if (Cell->IsLoaded())
-			{
-				CellsToRefresh.Add(Cell);
-			}
-		});
-		return !CellsToRefresh.IsEmpty();
-	};
-
-	return UpdateEditorCells(GetCellsToRefresh, /*bIsCellShouldBeLoaded*/true, bIsFromUserChange);
-}
-
 void UWorldPartition::UnloadEditorCells(const FBox& Box, bool bIsFromUserChange)
 {
 	FWorldPartionCellUpdateContext CellUpdateContext(this);
diff --git a/Engine/Source/Runtime/Engine/Public/WorldPartition/DataLayer/DataLayerSubsystem.h b/Engine/Source/Runtime/Engine/Public/WorldPartition/DataLayer/DataLayerSubsystem.h
--- a/Engine/Source/Runtime/Engine/Public/WorldPartition/DataLayer/DataLayerSubsystem.h
+++ b/Engine/Source/Runtime/Engine/Public/WorldPartition/DataLayer/DataLayerSubsystem.h
@@ -16,6 +16,25 @@ class UCanvas;
 
 DECLARE_DYNAMIC_MULTICAST_DELEGATE_TwoParams(FOnDataLayerRuntimeStateChanged, const UDataLayerInstance*, DataLayer, EDataLayerRuntimeState, State);
 
+#if WITH_EDITOR
+class ENGINE_API FDataLayersEditorBroadcast
+{
+	friend class UDataLayerEditorSubsystem;
+
+public:
+	static FDataLayersEditorBroadcast& Get();
+
+	static void StaticOnActorDataLayersEditorLoadingStateChanged(bool bIsFromUserChange);
+
+	/** Broadcasts whenever one or more DataLayers editor loading state changed */
+	DECLARE_EVENT_OneParam(FDataLayersEditorBroadcast, FOnActorDataLayersEditorLoadingStateChanged, bool /*bIsFromUserChange*/);
+	FOnActorDataLayersEditorLoadingStateChanged& OnActorDataLayersEditorLoadingStateChanged() { return DataLayerEditorLoadingStateChanged; }
+
+private:
+	FOnActorDataLayersEditorLoadingStateChanged DataLayerEditorLoadingStateChanged;
+};
+#endif
+
 UCLASS()
 class ENGINE_API UDataLayerSubsystem : public UWorldSubsystem
 {
diff --git a/Engine/Source/Runtime/Engine/Public/WorldPartition/WorldPartition.h b/Engine/Source/Runtime/Engine/Public/WorldPartition/WorldPartition.h
--- a/Engine/Source/Runtime/Engine/Public/WorldPartition/WorldPartition.h
+++ b/Engine/Source/Runtime/Engine/Public/WorldPartition/WorldPartition.h
@@ -147,7 +147,6 @@ public:
 	void LoadEditorCells(const TArray<FName>& CellNames, bool bIsFromUserChange);
 	void UnloadEditorCells(const FBox& Box, bool bIsFromUserChange);
 	bool AreEditorCellsLoaded(const FBox& Box);
-	bool RefreshLoadedEditorCells(bool bIsFromUserChange);
 
 	// PIE/Game/Cook Methods
 	bool GenerateStreaming(TArray<FString>* OutPackagesToGenerate = nullptr);
@@ -281,6 +280,8 @@ private:
 	void UnregisterDelegates();	
 
 #if WITH_EDITOR
+	void OnActorDataLayersEditorLoadingStateChanged(bool bFromUserOperation);
+
 	void UpdateLoadingEditorCell(UWorldPartitionEditorCell* Cell, bool bShouldBeLoaded, bool bFromUserOperation);
 	void HashActorDesc(FWorldPartitionActorDesc* ActorDesc);
 	void UnhashActorDesc(FWorldPartitionActorDesc* ActorDesc);
