Move DataLayerInstanceName resolving to UDataLayerSubsystem from UDataLayerEditorSubsystem.
UDataLayerEditorSubsystem is not initialized when lauching with "-game"

diff --git a/Engine/Source/Editor/DataLayerEditor/Private/DataLayer/DataLayerEditorSubsystem.cpp b/Engine/Source/Editor/DataLayerEditor/Private/DataLayer/DataLayerEditorSubsystem.cpp
--- a/Engine/Source/Editor/DataLayerEditor/Private/DataLayer/DataLayerEditorSubsystem.cpp
+++ b/Engine/Source/Editor/DataLayerEditor/Private/DataLayer/DataLayerEditorSubsystem.cpp
@@ -6,7 +6,6 @@
 #include "WorldPartition/DataLayer/DeprecatedDataLayerInstance.h"
 #include "WorldPartition/DataLayer/WorldDataLayers.h"
 #include "WorldPartition/DataLayer/DataLayerSubsystem.h"
-#include "WorldPartition/DataLayer/DataLayerUtils.h"
 #include "WorldPartition/WorldPartitionSubsystem.h"
 #include "WorldPartition/WorldPartition.h"
 #include "WorldPartition/WorldPartitionEditorPerProjectUserSettings.h"
@@ -25,7 +24,6 @@
 #include "EngineUtils.h"
 #include "Editor.h"
 #include "Subsystems/ActorEditorContextSubsystem.h"
-#include "ProfilingDebugging/ScopedTimers.h"
 
 #define LOCTEXT_NAMESPACE "DataLayer"
 
@@ -56,7 +54,6 @@ private:
 	void OnObjectPostEditChange(UObject* Object, FPropertyChangedEvent& PropertyChangedEvent);
 	void OnLevelActorsAdded(AActor* InActor) { DataLayerEditorSubsystem->InitializeNewActorDataLayers(InActor); }
 	void OnLevelSelectionChanged(UObject* InObject) { DataLayerEditorSubsystem->OnSelectionChanged(); }
-	void OnActorDescContainerInitialized(UActorDescContainer* InActorDescContainer);
 
 	UDataLayerEditorSubsystem* DataLayerEditorSubsystem;
 	bool bIsInitialized;
@@ -85,7 +82,6 @@ void FDataLayersBroadcast::Deinitialize()
 		GEngine->OnLevelActorAdded().RemoveAll(this);
 		USelection::SelectionChangedEvent.RemoveAll(this);
 		USelection::SelectObjectEvent.RemoveAll(this);
-		UActorDescContainer::OnActorDescContainerInitialized.RemoveAll(this);
 	}
 }
 
@@ -100,15 +96,9 @@ void FDataLayersBroadcast::Initialize()
 		GEngine->OnLevelActorAdded().AddRaw(this, &FDataLayersBroadcast::OnLevelActorsAdded);
 		USelection::SelectionChangedEvent.AddRaw(this, &FDataLayersBroadcast::OnLevelSelectionChanged);
 		USelection::SelectObjectEvent.AddRaw(this, &FDataLayersBroadcast::OnLevelSelectionChanged);
-		UActorDescContainer::OnActorDescContainerInitialized.AddRaw(this, &FDataLayersBroadcast::OnActorDescContainerInitialized);
 	}
 }
 
-void FDataLayersBroadcast::OnActorDescContainerInitialized(UActorDescContainer* InActorDescContainer)
-{
-	UDataLayerEditorSubsystem::OnActorDescContainerInitialized(InActorDescContainer);
-}
-
 void FDataLayersBroadcast::OnObjectPostEditChange(UObject* Object, FPropertyChangedEvent& PropertyChangedEvent)
 {
 	if (Object && (PropertyChangedEvent.ChangeType != EPropertyChangeType::Interactive))
@@ -181,21 +171,6 @@ void UDataLayerEditorSubsystem::Deinitialize()
 	DataLayerEditorLoadingStateChanged.Remove(OnActorDataLayersEditorLoadingStateChangedEngineBridgeHandle);
 }
 
-void UDataLayerEditorSubsystem::OnActorDescContainerInitialized(UActorDescContainer* InActorDescContainer)
-{
-	check(InActorDescContainer);
-
-	UE_SCOPED_TIMER(*FString::Printf(TEXT("Resolving Data Layer Instance Names for %s"), *InActorDescContainer->GetContainerPackage().ToString()), LogDataLayerEditorSubsystem, Display);
-
-	const FWorldDataLayersActorDesc* WorldDataLayersActorDesc = FDataLayerUtils::GetWorldDataLayersActorDesc(InActorDescContainer);
-	for (FActorDescList::TIterator<> Iterator(InActorDescContainer); Iterator; ++Iterator)
-	{
-		FWorldPartitionActorDesc* ActorDesc = *Iterator;
-		check(ActorDesc->GetContainer() == InActorDescContainer);
-		ActorDesc->DataLayerInstanceNames = FDataLayerUtils::ResolvedDataLayerInstanceNames(ActorDesc, WorldDataLayersActorDesc);
-	}
-}
-
 void UDataLayerEditorSubsystem::OnExecuteActorEditorContextAction(UWorld* InWorld, const EActorEditorContextAction& InType, AActor* InActor)
 {	
 	UDataLayerSubsystem* DataLayerSubsystem = GetWorld()->GetSubsystem<UDataLayerSubsystem>();
diff --git a/Engine/Source/Editor/DataLayerEditor/Public/DataLayer/DataLayerEditorSubsystem.h b/Engine/Source/Editor/DataLayerEditor/Public/DataLayer/DataLayerEditorSubsystem.h
--- a/Engine/Source/Editor/DataLayerEditor/Public/DataLayer/DataLayerEditorSubsystem.h
+++ b/Engine/Source/Editor/DataLayerEditor/Public/DataLayer/DataLayerEditorSubsystem.h
@@ -683,7 +683,5 @@ private:
 	/** Delegate used to notify changes to ActorEditorContextSubsystem */
 	FOnActorEditorContextClientChanged ActorEditorContextClientChanged;
 
-	static void OnActorDescContainerInitialized(class UActorDescContainer* InActorDescContainer);
-
 	friend class FDataLayersBroadcast;
 };
diff --git a/Engine/Source/Runtime/Engine/Private/WorldPartition/DataLayer/DataLayerSubsystem.cpp b/Engine/Source/Runtime/Engine/Private/WorldPartition/DataLayer/DataLayerSubsystem.cpp
--- a/Engine/Source/Runtime/Engine/Private/WorldPartition/DataLayer/DataLayerSubsystem.cpp
+++ b/Engine/Source/Runtime/Engine/Private/WorldPartition/DataLayer/DataLayerSubsystem.cpp
@@ -14,6 +14,8 @@
 #if WITH_EDITOR
 #include "Editor.h"
 #include "Modules/ModuleManager.h"
+#include "ProfilingDebugging/ScopedTimers.h"
+#include "WorldPartition/DataLayer/DataLayerUtils.h"
 #include "WorldPartition/DataLayer/IDataLayerEditorModule.h"
 #include "WorldPartition/WorldPartitionEditorPerProjectUserSettings.h"
 #endif
@@ -64,6 +66,17 @@ void UDataLayerSubsystem::Initialize(FSubsystemCollectionBase& Collection)
 	{
 		FModuleManager::LoadModuleChecked<IDataLayerEditorModule>("DataLayerEditor");
 	}
+
+	UActorDescContainer::OnActorDescContainerInitialized.AddUObject(this, &UDataLayerSubsystem::OnActorDescContainerInitialized);
+#endif
+}
+
+void UDataLayerSubsystem::Deinitialize()
+{
+	Super::Deinitialize();
+
+#if WITH_EDITOR
+	UActorDescContainer::OnActorDescContainerInitialized.RemoveAll(this);
 #endif
 }
 
@@ -610,6 +623,21 @@ uint32 UDataLayerSubsystem::GetDataLayerEditorContextHash() const
 	return FDataLayerEditorContext(GetWorld(), DataLayerInstanceNames).GetHash();
 }
 
+void UDataLayerSubsystem::OnActorDescContainerInitialized(UActorDescContainer* InActorDescContainer)
+{
+	check(InActorDescContainer);
+
+	UE_SCOPED_TIMER(*FString::Printf(TEXT("Resolving Data Layer Instance Names for %s"), *InActorDescContainer->GetContainerPackage().ToString()), LogWorldPartition, Display);
+
+	const FWorldDataLayersActorDesc* WorldDataLayersActorDesc = FDataLayerUtils::GetWorldDataLayersActorDesc(InActorDescContainer);
+	for (FActorDescList::TIterator<> Iterator(InActorDescContainer); Iterator; ++Iterator)
+	{
+		FWorldPartitionActorDesc* ActorDesc = *Iterator;
+		check(ActorDesc->GetContainer() == InActorDescContainer);
+		ActorDesc->DataLayerInstanceNames = FDataLayerUtils::ResolvedDataLayerInstanceNames(ActorDesc, WorldDataLayersActorDesc);
+	}
+}
+
 #endif
 
 //~ Begin Deprecated
diff --git a/Engine/Source/Runtime/Engine/Public/WorldPartition/DataLayer/DataLayerSubsystem.h b/Engine/Source/Runtime/Engine/Public/WorldPartition/DataLayer/DataLayerSubsystem.h
--- a/Engine/Source/Runtime/Engine/Public/WorldPartition/DataLayer/DataLayerSubsystem.h
+++ b/Engine/Source/Runtime/Engine/Public/WorldPartition/DataLayer/DataLayerSubsystem.h
@@ -13,6 +13,7 @@
  * UDataLayerSubsystem
  */
 
+class UActorDescContainer;
 class UCanvas;
 
 DECLARE_DYNAMIC_MULTICAST_DELEGATE_TwoParams(FOnDataLayerRuntimeStateChanged, const UDataLayerInstance*, DataLayer, EDataLayerRuntimeState, State);
@@ -46,6 +47,7 @@ public:
 
 	//~ Begin USubsystem Interface.
 	virtual void Initialize(FSubsystemCollectionBase& Collection) override;
+	virtual void Deinitialize() override;
 	//~ End USubsystem Interface.
 
 	//~ Begin Blueprint callable functions
@@ -191,6 +193,10 @@ public:
 
 private:
 
+#if WITH_EDITOR
+	void OnActorDescContainerInitialized(UActorDescContainer* InActorDescContainer);
+#endif
+
 	/** Console command used to toggle activation of a DataLayer */
 	static class FAutoConsoleCommand ToggleDataLayerActivation;
 
@@ -204,9 +210,15 @@ private:
 template<class T>
 UDataLayerInstance* UDataLayerSubsystem::GetDataLayerInstance(const T& InDataLayerIdentifier) const
 {
-	PRAGMA_DISABLE_DEPRECATION_WARNINGS
-	const UDataLayerInstance * DataLayerInstance = GetWorld()->GetWorldDataLayers()->GetDataLayerInstance(InDataLayerIdentifier);
-	PRAGMA_ENABLE_DEPRECATION_WARNINGS
+	const UDataLayerInstance* DataLayerInstance = nullptr;
+
+	if (AWorldDataLayers* WorldDataLayers = GetWorld()->GetWorldDataLayers())
+	{
+		PRAGMA_DISABLE_DEPRECATION_WARNINGS
+		DataLayerInstance = WorldDataLayers->GetDataLayerInstance(InDataLayerIdentifier);
+		PRAGMA_ENABLE_DEPRECATION_WARNINGS
+	}
+	
 	return const_cast<UDataLayerInstance*>(DataLayerInstance);
 }
 
@@ -214,11 +226,16 @@ template<class T>
 TArray<FName> UDataLayerSubsystem::GetDataLayerInstanceNames(const TArray<T>& InDataLayerIdentifiers) const
 {
 	TArray<FName> DataLayerInstanceNames;
-	DataLayerInstanceNames.Reserve(InDataLayerIdentifiers.Num());
 
-	PRAGMA_DISABLE_DEPRECATION_WARNINGS
-	DataLayerInstanceNames = GetWorld()->GetWorldDataLayers()->GetDataLayerInstanceNames(InDataLayerIdentifiers);
-	PRAGMA_ENABLE_DEPRECATION_WARNINGS
+	if (AWorldDataLayers* WorldDataLayers = GetWorld()->GetWorldDataLayers())
+	{
+		DataLayerInstanceNames.Reserve(InDataLayerIdentifiers.Num());
+
+		PRAGMA_DISABLE_DEPRECATION_WARNINGS
+		DataLayerInstanceNames = GetWorld()->GetWorldDataLayers()->GetDataLayerInstanceNames(InDataLayerIdentifiers);
+		PRAGMA_ENABLE_DEPRECATION_WARNINGS
+	}
+
 	return DataLayerInstanceNames;
 }
 
@@ -226,10 +243,15 @@ template<class T>
 TArray<const UDataLayerInstance*> UDataLayerSubsystem::GetDataLayerInstances(const TArray<T>& InDataLayerIdentifiers) const
 {
 	TArray<const UDataLayerInstance*> DataLayerInstances;
-	DataLayerInstances.Reserve(InDataLayerIdentifiers.Num());
 
-	PRAGMA_DISABLE_DEPRECATION_WARNINGS
-	DataLayerInstances = GetWorld()->GetWorldDataLayers()->GetDataLayerInstances(InDataLayerIdentifiers);
-	PRAGMA_ENABLE_DEPRECATION_WARNINGS
+	if (AWorldDataLayers* WorldDataLayers = GetWorld()->GetWorldDataLayers())
+	{
+		DataLayerInstances.Reserve(InDataLayerIdentifiers.Num());
+
+		PRAGMA_DISABLE_DEPRECATION_WARNINGS
+		DataLayerInstances = GetWorld()->GetWorldDataLayers()->GetDataLayerInstances(InDataLayerIdentifiers);
+		PRAGMA_ENABLE_DEPRECATION_WARNINGS
+	}
+
 	return DataLayerInstances;
 }
\ No newline at end of file
diff --git a/Engine/Source/Runtime/Engine/Public/WorldPartition/WorldPartitionActorDesc.h b/Engine/Source/Runtime/Engine/Public/WorldPartition/WorldPartitionActorDesc.h
--- a/Engine/Source/Runtime/Engine/Public/WorldPartition/WorldPartitionActorDesc.h
+++ b/Engine/Source/Runtime/Engine/Public/WorldPartition/WorldPartitionActorDesc.h
@@ -39,7 +39,7 @@ class ENGINE_API FWorldPartitionActorDesc
 	friend class UWorldPartition;
 	friend struct FWorldPartitionHandleImpl;
 	friend struct FWorldPartitionReferenceImpl;
-	friend class UDataLayerEditorSubsystem;
+	friend class UDataLayerSubsystem;
 	friend class FDataLayerUtils;
 
 public:
