Use WorldPartitionRuntimeCell data layers whenever possible when retrieving data layers for an actor
Remove check in DataLayerHierarchy that the world need to be teared down when removing actors. Callback occurs when unloading cells.

#test PIE, SIE, Game

diff --git a/Engine/Source/Editor/DataLayerEditor/Private/DataLayer/DataLayerHierarchy.cpp b/Engine/Source/Editor/DataLayerEditor/Private/DataLayer/DataLayerHierarchy.cpp
--- a/Engine/Source/Editor/DataLayerEditor/Private/DataLayer/DataLayerHierarchy.cpp
+++ b/Engine/Source/Editor/DataLayerEditor/Private/DataLayer/DataLayerHierarchy.cpp
@@ -305,9 +305,6 @@ void FDataLayerHierarchy::OnLevelActorsRemoved(const TArray<AActor*>& InActors)
 			{
 				if (Actor != nullptr && Actor->HasDataLayers())
 				{
-					// While exiting PIE/SIE, The actors in cells do not have a valid OwningWorld anymore, so AActor::GetDataLayerInstance() wont resolve below.
-					check(CurrentWorld->bIsTearingDown || Actor->GetWorld() != nullptr);
-
 					const TArray<const UDataLayerInstance*> DataLayerInstances = Actor->GetDataLayerInstances();
 					EventData.ItemIDs.Reserve(DataLayerInstances.Num());
 					for (const UDataLayerInstance* DataLayerInstance : DataLayerInstances)
diff --git a/Engine/Source/Runtime/Engine/Private/Actor.cpp b/Engine/Source/Runtime/Engine/Private/Actor.cpp
--- a/Engine/Source/Runtime/Engine/Private/Actor.cpp
+++ b/Engine/Source/Runtime/Engine/Private/Actor.cpp
@@ -5980,17 +5980,7 @@ bool AActor::HasDataLayers() const
 bool AActor::UseWorldPartitionRuntimeCellDataLayers() const
 {
 	ULevel* Level = GetLevel();
-	bool bIsActorInCell = Level != nullptr && Level->IsWorldPartitionRuntimeCell();
-
-#if WITH_EDITOR
-	// In PIE, SIE & -game, use the WorldPartitionRuntimeCell's data layers to emulate the behavior of a Cooked Game.
-	UWorld* World = GetWorld();
-	bool bShouldUseCell = (World != nullptr && World->IsGameWorld()) || UWorldPartition::IsSimulating();
-
-	return bIsActorInCell && bShouldUseCell;
-#else
-	return bIsActorInCell;
-#endif
+	return Level != nullptr && Level->IsWorldPartitionRuntimeCell();
 }
 
 
