DataLayer - Fix Duplicated keys for localization

#rnx

diff --git a/Engine/Source/Editor/DataLayerEditor/Private/DataLayer/DataLayerMode.cpp b/Engine/Source/Editor/DataLayerEditor/Private/DataLayer/DataLayerMode.cpp
--- a/Engine/Source/Editor/DataLayerEditor/Private/DataLayer/DataLayerMode.cpp
+++ b/Engine/Source/Editor/DataLayerEditor/Private/DataLayer/DataLayerMode.cpp
@@ -417,7 +417,7 @@ FSceneOutlinerDragValidationInfo FDataLayerMode::ValidateDrop(const ISceneOutlin
 
 		if (TargetDataLayer->IsLocked())
 		{
-			const FText Text = bMoveOperation ? LOCTEXT("CantMoveLockedDataLayer", "Can't move actors to locked Data Layer") : LOCTEXT("CantAssignLockedDataLayer", "Can't assign actors to locked Data Layer");
+			const FText Text = bMoveOperation ? LOCTEXT("CantMoveActorsToLockedDataLayer", "Can't move actors to locked Data Layer") : LOCTEXT("CantAssignActorsToLockedDataLayer", "Can't assign actors to locked Data Layer");
 			return FSceneOutlinerDragValidationInfo(ESceneOutlinerDropCompatibility::IncompatibleGeneric, Text);
 		}
 
