Added support of Data Layers inside Level Instances.

diff --git a/Engine/Source/Editor/DataLayerEditor/Private/DataLayer/DataLayerActorTreeItem.h b/Engine/Source/Editor/DataLayerEditor/Private/DataLayer/DataLayerActorTreeItem.h
--- a/Engine/Source/Editor/DataLayerEditor/Private/DataLayer/DataLayerActorTreeItem.h
+++ b/Engine/Source/Editor/DataLayerEditor/Private/DataLayer/DataLayerActorTreeItem.h
@@ -6,7 +6,9 @@
 #include "GameFramework/Actor.h"
 #include "SceneOutlinerFwd.h"
 #include "ActorTreeItem.h"
+#include "Engine/World.h"
 #include "WorldPartition/DataLayer/DataLayerInstance.h"
+#include "LevelInstance/LevelInstanceSubsystem.h"
 
 //////////////////////////////////////////////////////////////////////////
 // FDataLayerActorTreeItemData
@@ -35,6 +37,7 @@ public:
 		, DataLayerInstance(InData.DataLayerInstance)
 		, IDDataLayerActor(FDataLayerActorTreeItem::ComputeTreeItemID(Actor.Get(), DataLayerInstance.Get()))
 	{
+		UpdateDisplayStringInternal();
 	}
 
 	UDataLayerInstance* GetDataLayer() const { return DataLayerInstance.Get(); }
@@ -66,7 +69,28 @@ public:
 	virtual bool GetVisibility() const override { return false; }
 	/* End ISceneOutlinerTreeItem Implementation */
 
+protected:
+
+	virtual void UpdateDisplayString() override
+	{
+		UpdateDisplayStringInternal();
+	}
+
 private:
+
+	void UpdateDisplayStringInternal()
+	{
+		FActorTreeItem::UpdateDisplayString();
+
+		UWorld* OwningWorld = Actor.IsValid() ? Actor->GetWorld() : nullptr;
+		ULevel* Level = Actor.IsValid() ? Actor->GetLevel() : nullptr;
+		ULevelInstanceSubsystem* LevelInstanceSubsystem = UWorld::GetSubsystem<ULevelInstanceSubsystem>(OwningWorld);
+		if (LevelInstanceSubsystem && Level && (Level != OwningWorld->GetCurrentLevel()))
+		{
+			DisplayString = LevelInstanceSubsystem->PrefixWithParentLevelInstanceActorLabels(DisplayString, Actor->GetLevel());
+		}
+	}
+
 	TWeakObjectPtr<UDataLayerInstance> DataLayerInstance;
 	const uint32 IDDataLayerActor;
 };
\ No newline at end of file
diff --git a/Engine/Source/Editor/DataLayerEditor/Private/DataLayer/DataLayerEditorSubsystem.cpp b/Engine/Source/Editor/DataLayerEditor/Private/DataLayer/DataLayerEditorSubsystem.cpp
--- a/Engine/Source/Editor/DataLayerEditor/Private/DataLayer/DataLayerEditorSubsystem.cpp
+++ b/Engine/Source/Editor/DataLayerEditor/Private/DataLayer/DataLayerEditorSubsystem.cpp
@@ -24,9 +24,9 @@
 #include "EngineUtils.h"
 #include "Editor.h"
 #include "Subsystems/ActorEditorContextSubsystem.h"
+#include "Misc/ScopedSlowTask.h"
 #include "Styling/AppStyle.h"
 
-
 #define LOCTEXT_NAMESPACE "DataLayer"
 
 DEFINE_LOG_CATEGORY_STATIC(LogDataLayerEditorSubsystem, All, All);
@@ -152,7 +152,9 @@ void UDataLayerEditorSubsystem::Initialize(FSubsystemCollectionBase& Collection)
 
 	if (UWorld* World = GetWorld())
 	{
-		World->PersistentLevel->OnLoadedActorAddedToLevelEvent.AddLambda([this](AActor& InActor) { InitializeNewActorDataLayers(&InActor); });
+		World->PersistentLevel->OnLoadedActorAddedToLevelEvent.AddUObject(this, &UDataLayerEditorSubsystem::OnLoadedActorAddedToLevel);
+		World->OnWorldPartitionInitialized().AddUObject(this, &UDataLayerEditorSubsystem::OnWorldPartitionInitialized);
+		World->OnWorldPartitionUninitialized().AddUObject(this, &UDataLayerEditorSubsystem::OnWorldPartitionUninitialized);
 	}
 
 	UActorEditorContextSubsystem::Get()->RegisterClient(this);
@@ -175,6 +177,7 @@ void UDataLayerEditorSubsystem::Deinitialize()
 
 void UDataLayerEditorSubsystem::OnExecuteActorEditorContextAction(UWorld* InWorld, const EActorEditorContextAction& InType, AActor* InActor)
 {	
+	check(InWorld == GetWorld());
 	UDataLayerSubsystem* DataLayerSubsystem = GetWorld()->GetSubsystem<UDataLayerSubsystem>();
 
 	switch (InType)
@@ -269,7 +272,9 @@ void UDataLayerEditorSubsystem::EditorMapChange()
 {
 	if (UWorld * World = GetWorld())
 	{
-		World->PersistentLevel->OnLoadedActorAddedToLevelEvent.AddLambda([this](AActor& InActor) { InitializeNewActorDataLayers(&InActor); });
+		World->PersistentLevel->OnLoadedActorAddedToLevelEvent.AddUObject(this, &UDataLayerEditorSubsystem::OnLoadedActorAddedToLevel);
+		World->OnWorldPartitionInitialized().AddUObject(this, &UDataLayerEditorSubsystem::OnWorldPartitionInitialized);
+		World->OnWorldPartitionUninitialized().AddUObject(this, &UDataLayerEditorSubsystem::OnWorldPartitionUninitialized);
 	}
 	BroadcastDataLayerChanged(EDataLayerAction::Reset, NULL, NAME_None);
 	UpdateAllActorsVisibility(true, true);
@@ -295,7 +300,23 @@ void UDataLayerEditorSubsystem::PostUndoRedo()
 
 bool UDataLayerEditorSubsystem::IsActorValidForDataLayer(AActor* Actor)
 {
-	return Actor && Actor->SupportsDataLayer() && Actor->IsValidForDataLayer() && (Actor->GetLevel() == Actor->GetWorld()->PersistentLevel);
+	return Actor && Actor->SupportsDataLayer() && Actor->IsValidForDataLayer() && ((Actor->GetLevel() == Actor->GetWorld()->PersistentLevel) || Actor->GetLevel()->GetWorldDataLayers());
+}
+
+void UDataLayerEditorSubsystem::OnWorldPartitionInitialized(UWorldPartition* InWorldPartition)
+{
+	InWorldPartition->GetTypedOuter<UWorld>()->PersistentLevel->OnLoadedActorAddedToLevelEvent.AddUObject(this, &UDataLayerEditorSubsystem::OnLoadedActorAddedToLevel);
+	UpdateAllActorsVisibility(true, true);
+}
+
+void UDataLayerEditorSubsystem::OnWorldPartitionUninitialized(UWorldPartition* InWorldPartition)
+{
+	InWorldPartition->GetTypedOuter<UWorld>()->PersistentLevel->OnLoadedActorAddedToLevelEvent.RemoveAll(this);
+}
+
+void UDataLayerEditorSubsystem::OnLoadedActorAddedToLevel(AActor& InActor)
+{
+	InitializeNewActorDataLayers(&InActor);
 }
 
 void UDataLayerEditorSubsystem::InitializeNewActorDataLayers(AActor* Actor)
@@ -328,7 +349,7 @@ bool UDataLayerEditorSubsystem::SetParentDataLayer(UDataLayerInstance* DataLayer
 		UpdateAllActorsVisibility(true, true);
 		if (bIsLoaded != DataLayer->IsEffectiveLoadedInEditor())
 		{
-			BroadcastDataLayerEditorLoadingStateChanged(true);
+			OnDataLayerEditorLoadingStateChanged(true);
 		}
 		return true;
 	}
@@ -645,7 +666,8 @@ bool UDataLayerEditorSubsystem::UpdateActorVisibility(AActor* Actor, bool& bOutS
 	}
 
 	// If the actor doesn't belong to any DataLayers
-	if (!Actor->HasValidDataLayers())
+	TArray<const UDataLayerInstance*> DataLayerInstances = Actor->GetDataLayerInstances();
+	if (DataLayerInstances.IsEmpty())
 	{
 		// Actors that don't belong to any DataLayer shouldn't be hidden
 		bOutActorModified = Actor->SetIsHiddenEdLayer(false);
@@ -653,10 +675,9 @@ bool UDataLayerEditorSubsystem::UpdateActorVisibility(AActor* Actor, bool& bOutS
 	}
 
 	bool bActorBelongsToVisibleDataLayer = false;
-	UDataLayerSubsystem* DataLayerSubsystem = UWorld::GetSubsystem<UDataLayerSubsystem>(GetWorld());
-	DataLayerSubsystem->ForEachDataLayer([&Actor, &bOutActorModified, &bActorBelongsToVisibleDataLayer](UDataLayerInstance* DataLayer)
+	for (const UDataLayerInstance* DataLayerInstance : DataLayerInstances)
 	{
-		if (DataLayer->IsEffectiveVisible() && Actor->ContainsDataLayer(DataLayer))
+		if (DataLayerInstance->IsEffectiveVisible())
 		{
 			if (Actor->SetIsHiddenEdLayer(false))
 			{
@@ -664,10 +685,9 @@ bool UDataLayerEditorSubsystem::UpdateActorVisibility(AActor* Actor, bool& bOutS
 			}
 			// Stop, because we found at least one visible DataLayer the actor belongs to
 			bActorBelongsToVisibleDataLayer = true;
-			return false;
+			break;
+		}
 		}
-		return true;
-	});
 
 	// If the actor isn't part of a visible DataLayer, hide and de-select it.
 	if (!bActorBelongsToVisibleDataLayer)
@@ -900,22 +920,22 @@ bool UDataLayerEditorSubsystem::SetDataLayerIsLoadedInEditor(UDataLayerInstance*
 {
 	if (SetDataLayerIsLoadedInEditorInternal(DataLayer, bIsLoadedInEditor, bIsFromUserChange))
 	{
-		BroadcastDataLayerEditorLoadingStateChanged(bIsFromUserChange);
+		OnDataLayerEditorLoadingStateChanged(bIsFromUserChange);
 	}
 	return true;
 }
 
 bool UDataLayerEditorSubsystem::SetDataLayersIsLoadedInEditor(const TArray<UDataLayerInstance*>& DataLayers, const bool bIsLoadedInEditor, const bool bIsFromUserChange)
 {
-	bool bRefreshNeeded = false;
+	bool bChanged = false;
 	for (UDataLayerInstance* DataLayer : DataLayers)
 	{
-		bRefreshNeeded |= SetDataLayerIsLoadedInEditorInternal(DataLayer, bIsLoadedInEditor, bIsFromUserChange);
+		bChanged |= SetDataLayerIsLoadedInEditorInternal(DataLayer, bIsLoadedInEditor, bIsFromUserChange);
 	}
 	
-	if (bRefreshNeeded)
+	if (bChanged)
 	{
-		BroadcastDataLayerEditorLoadingStateChanged(bIsFromUserChange);
+		OnDataLayerEditorLoadingStateChanged(bIsFromUserChange);
 	}
 
 	return true;
@@ -929,15 +949,15 @@ bool UDataLayerEditorSubsystem::ToggleDataLayerIsLoadedInEditor(UDataLayerInstan
 
 bool UDataLayerEditorSubsystem::ToggleDataLayersIsLoadedInEditor(const TArray<UDataLayerInstance*>& DataLayers, const bool bIsFromUserChange)
 {
-	bool bRefreshNeeded = false;
+	bool bChanged = false;
 	for (UDataLayerInstance* DataLayer : DataLayers)
 	{
-		bRefreshNeeded |= SetDataLayerIsLoadedInEditorInternal(DataLayer, !DataLayer->IsLoadedInEditor(), bIsFromUserChange);
+		bChanged |= SetDataLayerIsLoadedInEditorInternal(DataLayer, !DataLayer->IsLoadedInEditor(), bIsFromUserChange);
 	}
 	
-	if (bRefreshNeeded)
+	if (bChanged)
 	{
-		BroadcastDataLayerEditorLoadingStateChanged(bIsFromUserChange);
+		OnDataLayerEditorLoadingStateChanged(bIsFromUserChange);
 	}
 
 	return true;
@@ -945,17 +965,17 @@ bool UDataLayerEditorSubsystem::ToggleDataLayersIsLoadedInEditor(const TArray<UD
 
 bool UDataLayerEditorSubsystem::ResetUserSettings()
 {
-	bool bRefreshNeeded = false;
+	bool bChanged = false;
 	UDataLayerSubsystem* DataLayerSubsystem = UWorld::GetSubsystem<UDataLayerSubsystem>(GetWorld());
-	DataLayerSubsystem->ForEachDataLayer([this, &bRefreshNeeded](UDataLayerInstance* DataLayer)
+	DataLayerSubsystem->ForEachDataLayer([this, &bChanged](UDataLayerInstance* DataLayer)
 	{
-		bRefreshNeeded |= SetDataLayerIsLoadedInEditorInternal(DataLayer, DataLayer->IsInitiallyLoadedInEditor(), true);
+		bChanged |= SetDataLayerIsLoadedInEditorInternal(DataLayer, DataLayer->IsInitiallyLoadedInEditor(), true);
 		return true;
 	});
 	
-	if (bRefreshNeeded)
+	if (bChanged)
 	{
-		BroadcastDataLayerEditorLoadingStateChanged(true);
+		OnDataLayerEditorLoadingStateChanged(true);
 	}
 
 	return true;
@@ -963,9 +983,9 @@ bool UDataLayerEditorSubsystem::ResetUserSettings()
 
 bool UDataLayerEditorSubsystem::HasDeprecatedDataLayers() const
 {
-	if (AWorldDataLayers* WorldDataLayer = GetWorld()->GetWorldDataLayers())
+	if (AWorldDataLayers* WorldDataLayers = GetWorld()->GetWorldDataLayers())
 	{
-		return WorldDataLayer->HasDeprecatedDataLayers();
+		return WorldDataLayers->HasDeprecatedDataLayers();
 	}
 
 	return false;
@@ -998,7 +1018,8 @@ UDataLayerInstance* UDataLayerEditorSubsystem::CreateDataLayerInstance(const FDa
 	UDataLayerInstance* NewDataLayer = nullptr;
 
 	AWorldDataLayers* WorldDataLayers = Parameters.WorlDataLayers != nullptr ? Parameters.WorlDataLayers.Get() : GetWorld()->GetWorldDataLayers();
-	if (!HasDeprecatedDataLayers())
+	check(!Parameters.ParentDataLayer || !Parameters.WorlDataLayers || (Parameters.ParentDataLayer->GetOuterAWorldDataLayers() == Parameters.WorlDataLayers));
+	if (!WorldDataLayers->HasDeprecatedDataLayers())
 	{
 		NewDataLayer = WorldDataLayers->CreateDataLayer<UDataLayerInstanceWithAsset>(Parameters.DataLayerAsset);
 	}
@@ -1047,6 +1068,14 @@ void UDataLayerEditorSubsystem::BroadcastDataLayerChanged(const EDataLayerAction
 	ActorEditorContextClientChanged.Broadcast(this);
 }
 
+void UDataLayerEditorSubsystem::OnDataLayerEditorLoadingStateChanged(bool bIsFromUserChange)
+{
+	FScopedSlowTask SlowTask(1, LOCTEXT("UpdatingLoadedActors", "Updating loaded actors..."));
+	SlowTask.MakeDialog();
+
+	BroadcastDataLayerEditorLoadingStateChanged(bIsFromUserChange);
+}
+
 void UDataLayerEditorSubsystem::BroadcastDataLayerEditorLoadingStateChanged(bool bIsFromUserChange)
 {
 	DataLayerEditorLoadingStateChanged.Broadcast(bIsFromUserChange);
diff --git a/Engine/Source/Editor/DataLayerEditor/Private/DataLayer/DataLayerHierarchy.cpp b/Engine/Source/Editor/DataLayerEditor/Private/DataLayer/DataLayerHierarchy.cpp
--- a/Engine/Source/Editor/DataLayerEditor/Private/DataLayer/DataLayerHierarchy.cpp
+++ b/Engine/Source/Editor/DataLayerEditor/Private/DataLayer/DataLayerHierarchy.cpp
@@ -5,17 +5,22 @@
 #include "DataLayerActorTreeItem.h"
 #include "DataLayersActorDescTreeItem.h"
 #include "DataLayerTreeItem.h"
+#include "WorldDataLayersTreeItem.h"
 #include "DataLayerMode.h"
 #include "DataLayer/DataLayerEditorSubsystem.h"
 #include "WorldPartition/DataLayer/DataLayerAsset.h"
 #include "WorldPartition/DataLayer/DataLayerSubsystem.h"
+#include "WorldPartition/DataLayer/DataLayerInstanceWithAsset.h"
+#include "WorldPartition/DataLayer/WorldDataLayers.h"
+#include "WorldPartition/WorldPartitionSubsystem.h"
+#include "WorldPartition/WorldPartitionHelpers.h"
+#include "WorldPartition/IWorldPartitionEditorModule.h"
+#include "LevelInstance/LevelInstanceSubsystem.h"
 #include "Engine/Engine.h"
 #include "Engine/World.h"
 #include "EngineUtils.h"
 #include "Editor.h"
 #include "Modules/ModuleManager.h"
-#include "WorldPartition/WorldPartitionHelpers.h"
-#include "WorldPartition/IWorldPartitionEditorModule.h"
 
 TUniquePtr<FDataLayerHierarchy> FDataLayerHierarchy::Create(FDataLayerMode* Mode, const TWeakObjectPtr<UWorld>& World)
 {
@@ -31,6 +36,7 @@ FDataLayerHierarchy::FDataLayerHierarchy(FDataLayerMode* Mode, const TWeakObject
 	, bShowUnloadedActors(true)
 	, bShowOnlySelectedActors(false)
 	, bHighlightSelectedDataLayers(false)
+	, bShowLevelInstanceContent(false)
 {
 	if (GEngine)
 	{
@@ -49,6 +55,9 @@ FDataLayerHierarchy::FDataLayerHierarchy(FDataLayerMode* Mode, const TWeakObject
 			World->PersistentLevel->OnLoadedActorAddedToLevelEvent.AddRaw(this, &FDataLayerHierarchy::OnLoadedActorAdded);
 			World->PersistentLevel->OnLoadedActorRemovedFromLevelEvent.AddRaw(this, &FDataLayerHierarchy::OnLoadedActorRemoved);
 		}
+
+		World->OnWorldPartitionInitialized().AddRaw(this, &FDataLayerHierarchy::OnWorldPartitionInitialized);
+		World->OnWorldPartitionUninitialized().AddRaw(this, &FDataLayerHierarchy::OnWorldPartitionUninitialized);
 		
 		if (UWorldPartition* WorldPartition = World->GetWorldPartition())
 		{
@@ -82,6 +91,9 @@ FDataLayerHierarchy::~FDataLayerHierarchy()
 			RepresentingWorld->PersistentLevel->OnLoadedActorAddedToLevelEvent.RemoveAll(this);
 			RepresentingWorld->PersistentLevel->OnLoadedActorRemovedFromLevelEvent.RemoveAll(this);
 		}
+
+		RepresentingWorld->OnWorldPartitionInitialized().RemoveAll(this);
+		RepresentingWorld->OnWorldPartitionUninitialized().RemoveAll(this);
 		
 		if (UWorldPartition* WorldPartition = RepresentingWorld->GetWorldPartition())
 		{
@@ -97,20 +109,20 @@ FDataLayerHierarchy::~FDataLayerHierarchy()
 }
 
 
-bool FDataLayerHierarchy::IsDataLayerPartOfSelection(const UDataLayerInstance* DataLayer) const
+bool FDataLayerHierarchy::IsDataLayerPartOfSelection(const UDataLayerInstance* DataLayerInstance) const
 {
 	if (!bShowOnlySelectedActors)
 	{
 		return true;
 	}
 
-	if (UDataLayerEditorSubsystem::Get()->DoesDataLayerContainSelectedActors(DataLayer))
+	if (UDataLayerEditorSubsystem::Get()->DoesDataLayerContainSelectedActors(DataLayerInstance))
 	{
 		return true;
 	}
 
 	bool bFoundSelected = false;
-	DataLayer->ForEachChild([this, &bFoundSelected](const UDataLayerInstance* Child)
+	DataLayerInstance->ForEachChild([this, &bFoundSelected](const UDataLayerInstance* Child)
 	{
 		bFoundSelected = IsDataLayerPartOfSelection(Child);
 		return !bFoundSelected; // Continue iterating if not found
@@ -128,37 +140,49 @@ FSceneOutlinerTreeItemPtr FDataLayerHierarchy::CreateDataLayerTreeItem(UDataLaye
 	return Item;
 }
 
+UWorld* FDataLayerHierarchy::GetOwningWorld() const
+{
+	UWorld* World = RepresentingWorld.Get();
+	return World ? World->PersistentLevel->GetWorld() : nullptr;
+}
+
 void FDataLayerHierarchy::CreateItems(TArray<FSceneOutlinerTreeItemPtr>& OutItems) const
 {
-	const UDataLayerSubsystem* DataLayerSubsystem = UWorld::GetSubsystem<UDataLayerSubsystem>(RepresentingWorld.Get());
-	if (!DataLayerSubsystem)
-	{
-		return;
-	}
+	const UDataLayerSubsystem* DataLayerSubsystem = UWorld::GetSubsystem<UDataLayerSubsystem>(GetOwningWorld());
 
-	auto IsDataLayerShown = [this](const UDataLayerInstance* DataLayer)
+	auto IsDataLayerShown = [this](const UDataLayerInstance* DataLayerInstance)
 	{
-		const bool bIsRuntimeDataLayer = DataLayer->IsRuntime();
-		return ((bIsRuntimeDataLayer && bShowRuntimeDataLayers) || (!bIsRuntimeDataLayer && bShowEditorDataLayers)) && IsDataLayerPartOfSelection(DataLayer);
+		const bool bIsRuntimeDataLayer = DataLayerInstance->IsRuntime();
+		return ((bIsRuntimeDataLayer && bShowRuntimeDataLayers) || (!bIsRuntimeDataLayer && bShowEditorDataLayers)) && IsDataLayerPartOfSelection(DataLayerInstance);
 	};
 
-	DataLayerSubsystem->ForEachDataLayer([this, &OutItems, IsDataLayerShown](UDataLayerInstance* DataLayer)
+	auto CreateDataLayerTreeItems = [this, DataLayerSubsystem, &OutItems, IsDataLayerShown](const ULevel* OuterLevel)
 	{
-		if (IsDataLayerShown(DataLayer))
+		DataLayerSubsystem->ForEachDataLayer([this, &OutItems, IsDataLayerShown](UDataLayerInstance* DataLayer)
 		{
-			if (FSceneOutlinerTreeItemPtr DataLayerItem = CreateDataLayerTreeItem(DataLayer))
+			if (IsDataLayerShown(DataLayer))
 			{
-				OutItems.Add(DataLayerItem);
+				if (FSceneOutlinerTreeItemPtr DataLayerItem = CreateDataLayerTreeItem(DataLayer))
+				{
+					OutItems.Add(DataLayerItem);
+				}
 			}
-		}
-		return true;
-	});
+			return true;
+		}, OuterLevel);
+	};
+
+	// CurrentLevel represents the current level if different than the PersistentLevel
+	ULevel* CurrentLevel = (GetOwningWorld()->GetCurrentLevel() && !GetOwningWorld()->GetCurrentLevel()->IsPersistentLevel()) ? GetOwningWorld()->GetCurrentLevel() : nullptr;
+
+	CreateDataLayerTreeItems(GetOwningWorld()->PersistentLevel);
+	CreateDataLayerTreeItems(CurrentLevel);
 
 	if (bShowDataLayerActors)
 	{
-		for (AActor* Actor : FActorRange(RepresentingWorld.Get()))
+		for (AActor* Actor : FActorRange(GetOwningWorld()))
 		{
-			if (Actor->HasDataLayers())
+			// Only consider actors of current level (if there is one)
+			if (!CurrentLevel || (Actor->GetLevel() == CurrentLevel))
 			{
 				for (const UDataLayerInstance* DataLayerInstance : Actor->GetDataLayerInstances())
 				{
@@ -175,26 +199,42 @@ void FDataLayerHierarchy::CreateItems(TArray<FSceneOutlinerTreeItemPtr>& OutItem
 
 		if (bShowUnloadedActors)
 		{
-			if (UWorldPartition* const WorldPartition = RepresentingWorld.Get()->GetWorldPartition())
+			ULevelInstanceSubsystem* LevelInstanceSubsystem = UWorld::GetSubsystem<ULevelInstanceSubsystem>(GetOwningWorld());
+			if (UWorldPartitionSubsystem* WorldPartitionSubsystem = UWorld::GetSubsystem<UWorldPartitionSubsystem>(GetOwningWorld()))
 			{
-				FWorldPartitionHelpers::ForEachActorDesc(WorldPartition, [this, WorldPartition, DataLayerSubsystem, IsDataLayerShown, &OutItems](const FWorldPartitionActorDesc* ActorDesc)
+				WorldPartitionSubsystem->ForEachWorldPartition([this, CurrentLevel, LevelInstanceSubsystem, DataLayerSubsystem, IsDataLayerShown, &OutItems](UWorldPartition* WorldPartition)
 				{
-					if (ActorDesc != nullptr && !ActorDesc->IsLoaded())
+					// Skip WorldPartition if it's not the one of the current level (the editing level instance)
+					// or if we hide the content of level instances
+					UWorld* OuterWorld = WorldPartition->GetTypedOuter<UWorld>();
+					ULevel* OuterLevel = OuterWorld ? OuterWorld->PersistentLevel : nullptr;
+					if ((CurrentLevel && (CurrentLevel != OuterLevel)) ||
+						(!bShowLevelInstanceContent && OuterLevel && LevelInstanceSubsystem && LevelInstanceSubsystem->GetOwningLevelInstance(OuterLevel)))
+					{
+						return true;
+					}
+
+					// Create an FDataLayerActorDescTreeItem for each unloaded actor of this WorldPartition
+					FWorldPartitionHelpers::ForEachActorDesc(WorldPartition, [this, DataLayerSubsystem, IsDataLayerShown, &OutItems](const FWorldPartitionActorDesc* ActorDesc)
 					{
-						for (const FName& DataLayerInstanceName : ActorDesc->GetDataLayerInstanceNames())
+						if (ActorDesc != nullptr && !ActorDesc->IsLoaded())
 						{
-							if (const UDataLayerInstance* const DataLayerInstance = DataLayerSubsystem->GetDataLayerInstance(DataLayerInstanceName))
+							for (const FName& DataLayerInstanceName : ActorDesc->GetDataLayerInstanceNames())
 							{
-								if (IsDataLayerShown(DataLayerInstance))
+								if (const UDataLayerInstance* const DataLayerInstance = DataLayerSubsystem->GetDataLayerInstance(DataLayerInstanceName))
 								{
-									if (const FSceneOutlinerTreeItemPtr ActorDescItem = Mode->CreateItemFor<FDataLayerActorDescTreeItem>(FDataLayerActorDescTreeItemData(ActorDesc->GetGuid(), WorldPartition, const_cast<UDataLayerInstance*>(DataLayerInstance))))
+									if (IsDataLayerShown(DataLayerInstance))
 									{
-										OutItems.Add(ActorDescItem);
+										if (const FSceneOutlinerTreeItemPtr ActorDescItem = Mode->CreateItemFor<FDataLayerActorDescTreeItem>(FDataLayerActorDescTreeItemData(ActorDesc->GetGuid(), ActorDesc->GetContainer(), const_cast<UDataLayerInstance*>(DataLayerInstance))))
+										{
+											OutItems.Add(ActorDescItem);
+										}
 									}
 								}
 							}
 						}
-					}
+						return true;
+					});
 					return true;
 				});
 			}
@@ -206,9 +246,9 @@ FSceneOutlinerTreeItemPtr FDataLayerHierarchy::FindOrCreateParentItem(const ISce
 {
 	if (const FDataLayerTreeItem* DataLayerTreeItem = Item.CastTo<FDataLayerTreeItem>())
 	{
-		if (UDataLayerInstance* DataLayer = DataLayerTreeItem->GetDataLayer())
+		if (UDataLayerInstance* DataLayerInstance = DataLayerTreeItem->GetDataLayer())
 		{
-			if (UDataLayerInstance* ParentDataLayer = DataLayer->GetParent())
+			if (UDataLayerInstance* ParentDataLayer = DataLayerInstance->GetParent())
 			{
 				if (const FSceneOutlinerTreeItemPtr* ParentItem = Items.Find(ParentDataLayer))
 				{
@@ -219,33 +259,48 @@ FSceneOutlinerTreeItemPtr FDataLayerHierarchy::FindOrCreateParentItem(const ISce
 					return CreateDataLayerTreeItem(ParentDataLayer, true);
 				}
 			}
+			else
+			{
+				// Parent WorldDataLayers
+				if (AWorldDataLayers* OuterWorldDataLayers = DataLayerInstance->GetOuterAWorldDataLayers())
+				{
+					if (const FSceneOutlinerTreeItemPtr* ParentItem = Items.Find(OuterWorldDataLayers))
+					{
+						return *ParentItem;
+					}
+					else if (bCreate)
+					{
+						return Mode->CreateItemFor<FWorldDataLayersTreeItem>(OuterWorldDataLayers, true);
+					}
+				}
+			}
 		}
 	}
 	else if (const FDataLayerActorTreeItem* DataLayerActorTreeItem = Item.CastTo<FDataLayerActorTreeItem>())
 	{
-		if (UDataLayerInstance* DataLayer = DataLayerActorTreeItem->GetDataLayer())
+		if (UDataLayerInstance* DataLayerInstance = DataLayerActorTreeItem->GetDataLayer())
 		{
-			if (const FSceneOutlinerTreeItemPtr* ParentItem = Items.Find(DataLayer))
+			if (const FSceneOutlinerTreeItemPtr* ParentItem = Items.Find(DataLayerInstance))
 			{
 				return *ParentItem;
 			}
 			else if (bCreate)
 			{
-				return CreateDataLayerTreeItem(DataLayer, true);
+				return CreateDataLayerTreeItem(DataLayerInstance, true);
 			}
 		}
 	}
 	else if (const FDataLayerActorDescTreeItem* DataLayerActorDescTreeItem = Item.CastTo<FDataLayerActorDescTreeItem>())
 	{
-		if (UDataLayerInstance* DataLayer = DataLayerActorDescTreeItem->GetDataLayer())
+		if (UDataLayerInstance* DataLayerInstance = DataLayerActorDescTreeItem->GetDataLayer())
 		{
-			if (const FSceneOutlinerTreeItemPtr* ParentItem = Items.Find(DataLayer))
+			if (const FSceneOutlinerTreeItemPtr* ParentItem = Items.Find(DataLayerInstance))
 			{
 				return *ParentItem;
 			}
 			else if (bCreate)
 			{
-				return CreateDataLayerTreeItem(DataLayer, true);
+				return CreateDataLayerTreeItem(DataLayerInstance, true);
 			}
 		}
 	}
@@ -267,14 +322,14 @@ void FDataLayerHierarchy::OnLevelActorsAdded(const TArray<AActor*>& InActors)
 		return;
 	}
 
-	if (UWorld* CurrentWorld = RepresentingWorld.Get())
+	if (UWorld* OwningWorld = GetOwningWorld())
 	{
 		FSceneOutlinerHierarchyChangedData EventData;
 		EventData.Type = FSceneOutlinerHierarchyChangedData::Added;
 
 		for (AActor* Actor : InActors)
 		{
-			if (Actor != nullptr && Actor->HasDataLayers() && Actor->GetWorld() == CurrentWorld)
+			if (Actor && (Actor->GetWorld() == OwningWorld))
 			{
 				TArray<const UDataLayerInstance*> DataLayerInstances = Actor->GetDataLayerInstances();
 				EventData.Items.Reserve(DataLayerInstances.Num());
@@ -294,17 +349,16 @@ void FDataLayerHierarchy::OnLevelActorsAdded(const TArray<AActor*>& InActors)
 
 void FDataLayerHierarchy::OnLevelActorsRemoved(const TArray<AActor*>& InActors)
 {
-	if(UWorld* CurrentWorld = RepresentingWorld.Get())
+	if (UWorld* OwningWorld = GetOwningWorld())
 	{
 		FSceneOutlinerHierarchyChangedData EventData;
 		EventData.Type = FSceneOutlinerHierarchyChangedData::Removed;
 
 		for (AActor* Actor : InActors)
 		{
-			if (Actor != nullptr && Actor->HasDataLayers())
+			if (Actor != nullptr)
 			{
 				const TArray<const UDataLayerInstance*> DataLayerInstances = Actor->GetDataLayerInstances();
-				EventData.ItemIDs.Reserve(DataLayerInstances.Num());
 				for (const UDataLayerInstance* DataLayerInstance : DataLayerInstances)
 				{
 					EventData.ItemIDs.Add(FDataLayerActorTreeItem::ComputeTreeItemID(Actor, DataLayerInstance));
@@ -327,7 +381,7 @@ void FDataLayerHierarchy::OnLevelActorAdded(AActor* InActor)
 void FDataLayerHierarchy::OnActorDataLayersChanged(const TWeakObjectPtr<AActor>& InActor)
 {
 	AActor* Actor = InActor.Get();
-	if (Actor && RepresentingWorld.Get() == Actor->GetWorld())
+	if (Actor && (GetOwningWorld() == Actor->GetWorld()))
 	{
 		FullRefreshEvent();
 	}
@@ -335,8 +389,8 @@ void FDataLayerHierarchy::OnActorDataLayersChanged(const TWeakObjectPtr<AActor>&
 
 void FDataLayerHierarchy::OnDataLayerChanged(const EDataLayerAction Action, const TWeakObjectPtr<const UDataLayerInstance>& ChangedDataLayer, const FName& ChangedProperty)
 {
-	const UDataLayerInstance* DataLayer = ChangedDataLayer.Get();
-	if ((DataLayer && (RepresentingWorld.Get() == DataLayer->GetWorld())) || (Action == EDataLayerAction::Delete) || (Action == EDataLayerAction::Reset))
+	const UDataLayerInstance* DataLayerInstance = ChangedDataLayer.Get();
+	if (DataLayerInstance || (Action == EDataLayerAction::Delete) || (Action == EDataLayerAction::Reset))
 	{
 		FullRefreshEvent();
 	}
@@ -352,9 +406,19 @@ void FDataLayerHierarchy::OnLevelActorListChanged()
 	FullRefreshEvent();
 }
 
+void FDataLayerHierarchy::OnWorldPartitionInitialized(UWorldPartition* InWorldPartition)
+{
+	FullRefreshEvent();
+}
+
+void FDataLayerHierarchy::OnWorldPartitionUninitialized(UWorldPartition* InWorldPartition)
+{
+	FullRefreshEvent();
+}
+
 void FDataLayerHierarchy::OnLevelAdded(ULevel* InLevel, UWorld* InWorld)
 {
-	if (InLevel != nullptr && RepresentingWorld.Get() == InWorld)
+	if (!InWorld->IsGameWorld() && InLevel && (GetOwningWorld() == InWorld))
 	{
 		OnLevelActorsAdded(InLevel->Actors);
 	}
@@ -362,7 +426,7 @@ void FDataLayerHierarchy::OnLevelAdded(ULevel* InLevel, UWorld* InWorld)
 
 void FDataLayerHierarchy::OnLevelRemoved(ULevel* InLevel, UWorld* InWorld)
 {
-	if (InLevel != nullptr && RepresentingWorld.Get() == InWorld)
+	if (!InWorld->IsGameWorld() && InLevel && (GetOwningWorld() == InWorld))
 	{
 		OnLevelActorsRemoved(InLevel->Actors);
 	}
@@ -379,18 +443,23 @@ void FDataLayerHierarchy::OnLoadedActorAdded(AActor& InActor)
 	OnLevelActorAdded(&InActor);
 
 	// Remove corresponding actor desc items
-	if (RepresentingWorld.Get() == InActor.GetWorld() && InActor.HasDataLayers())
+	if ((GetOwningWorld() == InActor.GetWorld()))
 	{
-		FSceneOutlinerHierarchyChangedData EventData;
-		EventData.Type = FSceneOutlinerHierarchyChangedData::Removed;
-		
 		const TArray<const UDataLayerInstance*> DataLayerInstances = InActor.GetDataLayerInstances();
-		EventData.ItemIDs.Reserve(DataLayerInstances.Num());
-		for (const UDataLayerInstance* DataLayerInstance : DataLayerInstances)
+		if (DataLayerInstances.Num() > 0)
 		{
-			EventData.ItemIDs.Add(FDataLayerActorDescTreeItem::ComputeTreeItemID(InActor.GetActorGuid(), DataLayerInstance));
+			FSceneOutlinerHierarchyChangedData EventData;
+			EventData.Type = FSceneOutlinerHierarchyChangedData::Removed;
+			EventData.ItemIDs.Reserve(DataLayerInstances.Num());
+			ULevelInstanceSubsystem* LevelInstanceSubsystem = UWorld::GetSubsystem<ULevelInstanceSubsystem>(GetOwningWorld());
+			TArray<AActor*> ParentActors = LevelInstanceSubsystem->GetParentLevelInstanceActors(InActor.GetLevel());
+
+			for (const UDataLayerInstance* DataLayerInstance : DataLayerInstances)
+			{
+				EventData.ItemIDs.Add(FDataLayerActorDescTreeItem::ComputeTreeItemID(InActor.GetActorGuid(), DataLayerInstance, ParentActors));
+			}
+			HierarchyChangedEvent.Broadcast(EventData);
 		}
-		HierarchyChangedEvent.Broadcast(EventData);
 	}
 }
 
@@ -400,19 +469,21 @@ void FDataLayerHierarchy::OnLoadedActorRemoved(AActor& InActor)
 	OnLevelActorDeleted(&InActor);
 
 	// Add any corresponding actor desc items for this actor
-	UWorldPartition* const WorldPartition = RepresentingWorld.Get()->GetWorldPartition();
-	if (RepresentingWorld.Get() == InActor.GetWorld() && InActor.HasDataLayers() && WorldPartition != nullptr)
+	UWorldPartition* const WorldPartition = RepresentingWorld->GetWorldPartition();
+	if ((GetOwningWorld() == InActor.GetWorld()) && WorldPartition != nullptr)
 	{
-		FSceneOutlinerHierarchyChangedData EventData;
-		EventData.Type = FSceneOutlinerHierarchyChangedData::Added;
-
 		const TArray<const UDataLayerInstance*> DataLayerInstances = InActor.GetDataLayerInstances();
-		EventData.Items.Reserve(DataLayerInstances.Num());
-		for (const UDataLayerInstance* DataLayerInstance : DataLayerInstances)
+		if (DataLayerInstances.Num() > 0)
 		{
-			EventData.Items.Add(Mode->CreateItemFor<FDataLayerActorDescTreeItem>(FDataLayerActorDescTreeItemData(InActor.GetActorGuid(), WorldPartition, const_cast<UDataLayerInstance*>(DataLayerInstance))));
+			FSceneOutlinerHierarchyChangedData EventData;
+			EventData.Type = FSceneOutlinerHierarchyChangedData::Added;
+			EventData.Items.Reserve(DataLayerInstances.Num());
+			for (const UDataLayerInstance* DataLayerInstance : DataLayerInstances)
+			{
+				EventData.Items.Add(Mode->CreateItemFor<FDataLayerActorDescTreeItem>(FDataLayerActorDescTreeItemData(InActor.GetActorGuid(), WorldPartition, const_cast<UDataLayerInstance*>(DataLayerInstance))));
+			}
+			HierarchyChangedEvent.Broadcast(EventData);
 		}
-		HierarchyChangedEvent.Broadcast(EventData);
     }
 }
 
@@ -423,11 +494,11 @@ void FDataLayerHierarchy::OnActorDescAdded(FWorldPartitionActorDesc* InActorDesc
 		return;
 	}
 
-	UWorldPartition* const WorldPartition = RepresentingWorld.Get()->GetWorldPartition();
-	const UDataLayerSubsystem* const DataLayerSubsystem = UWorld::GetSubsystem<UDataLayerSubsystem>(RepresentingWorld.Get());
+	UWorldPartition* const WorldPartition = RepresentingWorld->GetWorldPartition();
+	const UDataLayerSubsystem* const DataLayerSubsystem = UWorld::GetSubsystem<UDataLayerSubsystem>(GetOwningWorld());
 	const TArray<FName>& DataLayerInstanceNames = InActorDesc->GetDataLayerInstanceNames();
 
-	if (DataLayerSubsystem != nullptr && RepresentingWorld->GetWorldPartition() == WorldPartition && DataLayerInstanceNames.Num() > 0)
+	if (DataLayerSubsystem && DataLayerInstanceNames.Num() > 0)
 	{
 		FSceneOutlinerHierarchyChangedData EventData;
 		EventData.Type = FSceneOutlinerHierarchyChangedData::Added;
@@ -449,19 +520,20 @@ void FDataLayerHierarchy::OnActorDescRemoved(FWorldPartitionActorDesc* InActorDe
 		return;
 	}
 
-	const UDataLayerSubsystem* const DataLayerSubsystem = UWorld::GetSubsystem<UDataLayerSubsystem>(RepresentingWorld.Get());
+	const UDataLayerSubsystem* const DataLayerSubsystem = UWorld::GetSubsystem<UDataLayerSubsystem>(GetOwningWorld());
 	const TArray<FName>& DataLayerInstanceNames = InActorDesc->GetDataLayerInstanceNames();
 	
-	if (DataLayerSubsystem != nullptr && DataLayerInstanceNames.Num() > 0)
+	if (DataLayerSubsystem && DataLayerInstanceNames.Num() > 0)
 	{
 		FSceneOutlinerHierarchyChangedData EventData;
 		EventData.Type = FSceneOutlinerHierarchyChangedData::Removed;
 		EventData.ItemIDs.Reserve(DataLayerInstanceNames.Num());
+		TArray<AActor*> ParentActors = FDataLayerActorDescTreeItem::GetParentActors(InActorDesc->GetContainer());
 
 		for (const FName& DataLayerInstanceName : DataLayerInstanceNames)
 		{
 			const UDataLayerInstance* const DataLayer = DataLayerSubsystem->GetDataLayerInstance(DataLayerInstanceName);
-			EventData.ItemIDs.Add(FDataLayerActorDescTreeItem::ComputeTreeItemID(InActorDesc->GetGuid(), DataLayer));
+			EventData.ItemIDs.Add(FDataLayerActorDescTreeItem::ComputeTreeItemID(InActorDesc->GetGuid(), DataLayer, ParentActors));
 		}
 		HierarchyChangedEvent.Broadcast(EventData);
 	}
diff --git a/Engine/Source/Editor/DataLayerEditor/Private/DataLayer/DataLayerHierarchy.h b/Engine/Source/Editor/DataLayerEditor/Private/DataLayer/DataLayerHierarchy.h
--- a/Engine/Source/Editor/DataLayerEditor/Private/DataLayer/DataLayerHierarchy.h
+++ b/Engine/Source/Editor/DataLayerEditor/Private/DataLayer/DataLayerHierarchy.h
@@ -11,6 +11,7 @@ class FDataLayerMode;
 class UDataLayerInstance;
 class UWorld;
 class FWorldPartitionActorDesc;
+class UWorldPartition;
 
 class FDataLayerHierarchy : public ISceneOutlinerHierarchy
 {
@@ -26,8 +27,10 @@ public:
 	void SetShowUnloadedActors(bool bInShowUnloadedActors) { bShowUnloadedActors = bInShowUnloadedActors; }
 	void SetShowOnlySelectedActors(bool bInbShowOnlySelectedActors) { bShowOnlySelectedActors = bInbShowOnlySelectedActors; }
 	void SetHighlightSelectedDataLayers(bool bInHighlightSelectedDataLayers) { bHighlightSelectedDataLayers = bInHighlightSelectedDataLayers; }
+	void SetShowLevelInstanceContent(bool bInShowLevelInstanceContent) { bShowLevelInstanceContent = bInShowLevelInstanceContent; }
 
 private:
+	UWorld* GetOwningWorld() const;
 	FDataLayerHierarchy(FDataLayerMode* Mode, const TWeakObjectPtr<UWorld>& Worlds);
 	FDataLayerHierarchy(const FDataLayerHierarchy&) = delete;
 	FDataLayerHierarchy& operator=(const FDataLayerHierarchy&) = delete;
@@ -38,6 +41,8 @@ private:
 	void OnLevelActorAdded(AActor* InActor);
 	void OnLevelActorDeleted(AActor* InActor);
 	void OnLevelActorListChanged();
+	void OnWorldPartitionInitialized(UWorldPartition* InWorldPartition);
+	void OnWorldPartitionUninitialized(UWorldPartition* InWorldPartition);
 	void OnLevelAdded(ULevel* InLevel, UWorld* InWorld);
 	void OnLevelRemoved(ULevel* InLevel, UWorld* InWorld);
 	void OnLoadedActorAdded(AActor& InActor);
@@ -57,4 +62,5 @@ private:
 	bool bShowUnloadedActors;
 	bool bShowOnlySelectedActors;
 	bool bHighlightSelectedDataLayers;
+	bool bShowLevelInstanceContent;
 };
\ No newline at end of file
diff --git a/Engine/Source/Editor/DataLayerEditor/Private/DataLayer/DataLayerMode.cpp b/Engine/Source/Editor/DataLayerEditor/Private/DataLayer/DataLayerMode.cpp
--- a/Engine/Source/Editor/DataLayerEditor/Private/DataLayer/DataLayerMode.cpp
+++ b/Engine/Source/Editor/DataLayerEditor/Private/DataLayer/DataLayerMode.cpp
@@ -9,12 +9,16 @@
 #include "DataLayer/DataLayerOutlinerDeleteButtonColumn.h"
 #include "WorldPartition/DataLayer/DataLayerSubsystem.h"
 #include "WorldPartition/DataLayer/DataLayerAsset.h"
+#include "WorldPartition/DataLayer/DataLayerInstanceWithAsset.h"
+#include "WorldPartition/WorldPartitionEditorPerProjectUserSettings.h"
 #include "DataLayer/SDataLayerOutliner.h"
 #include "DataLayerHierarchy.h"
 #include "DataLayerActorTreeItem.h"
 #include "DataLayerTreeItem.h"
+#include "WorldDataLayersTreeItem.h"
 #include "WorldTreeItem.h"
-#include "WorldPartition/WorldPartitionEditorPerProjectUserSettings.h"
+#include "LevelInstance/LevelInstanceEditorInstanceActor.h"
+#include "LevelInstance/LevelInstanceInterface.h"
 #include "ActorDescTreeItem.h"
 #include "ISceneOutlinerHierarchy.h"
 #include "SceneOutlinerMenuContext.h"
@@ -60,6 +64,7 @@ FDataLayerMode::FDataLayerMode(const FDataLayerModeParams& Params)
 	bHideUnloadedActors = SharedSettings->bHideUnloadedActors;
 	bShowOnlySelectedActors = SharedSettings->bShowOnlySelectedActors;
 	bHighlightSelectedDataLayers = SharedSettings->bHighlightSelectedDataLayers;
+	bHideLevelInstanceContent = SharedSettings->bHideLevelInstanceContent;
 
 	FSceneOutlinerFilterInfo ShowOnlySelectedActorsInfo(LOCTEXT("ToggleShowOnlySelected", "Only Selected"), LOCTEXT("ToggleShowOnlySelectedToolTip", "When enabled, only displays actors that are currently selected."), bShowOnlySelectedActors, FCreateSceneOutlinerFilter::CreateStatic(&FDataLayerMode::CreateShowOnlySelectedActorsFilter));
 	ShowOnlySelectedActorsInfo.OnToggle().AddLambda([this](bool bIsActive)
@@ -138,6 +143,39 @@ FDataLayerMode::FDataLayerMode(const FDataLayerModeParams& Params)
 		return FActorMode::IsActorDisplayable(SceneOutliner, Actor);
 	}), FSceneOutlinerFilter::EDefaultBehaviour::Pass));
 
+	FSceneOutlinerFilterInfo HideLevelInstancesInfo(LOCTEXT("ToggleHideLevelInstanceContent", "Hide Level Instance Content"), LOCTEXT("ToggleHideLevelInstanceContentToolTip", "When enabled, hides all level instance content."), bHideLevelInstanceContent, FCreateSceneOutlinerFilter::CreateStatic(&FDataLayerMode::CreateHideLevelInstancesFilter));
+	HideLevelInstancesInfo.OnToggle().AddLambda([this](bool bIsActive)
+	{
+		UWorldPartitionEditorPerProjectUserSettings* Settings = GetMutableDefault<UWorldPartitionEditorPerProjectUserSettings>();
+		Settings->bHideLevelInstanceContent = bHideLevelInstanceContent = bIsActive;
+		Settings->PostEditChange();
+
+		if (auto DataLayerHierarchy = StaticCast<FDataLayerHierarchy*>(Hierarchy.Get()))
+		{
+			DataLayerHierarchy->SetShowLevelInstanceContent(!bIsActive);
+		}
+	});
+	FilterInfoMap.Add(TEXT("HideLevelInstancesFilter"), HideLevelInstancesInfo);
+
+	// Add a filter which sets the interactive mode of LevelInstance items and their children
+	SceneOutliner->AddFilter(MakeShared<FDataLayerActorFilter>(FDataLayerActorTreeItem::FFilterPredicate::CreateStatic([](const AActor* Actor, const UDataLayerInstance* DataLayer) {return true; }), FSceneOutlinerFilter::EDefaultBehaviour::Pass, FDataLayerActorTreeItem::FFilterPredicate::CreateLambda([this](const AActor* Actor, const UDataLayerInstance* DataLayer)
+	{
+		if (!bHideLevelInstanceContent)
+		{
+			if (const ULevelInstanceSubsystem* LevelInstanceSubsystem = UWorld::GetSubsystem<ULevelInstanceSubsystem>(RepresentingWorld.Get()))
+			{
+				const ILevelInstanceInterface* ActorAsLevelInstance = Cast<ILevelInstanceInterface>(Actor);
+				const ILevelInstanceInterface* ActorParentLevelInstance = LevelInstanceSubsystem->GetParentLevelInstance(Actor);
+				if (!LevelInstanceSubsystem->IsEditingLevelInstance(ActorAsLevelInstance) &&
+					!LevelInstanceSubsystem->IsEditingLevelInstance(ActorParentLevelInstance))
+				{
+					return false;
+				}
+			}
+		}
+		return true;
+	})));
+
 	DataLayerEditorSubsystem = UDataLayerEditorSubsystem::Get();
 	Rebuild();
 	const_cast<FSharedSceneOutlinerData&>(SceneOutliner->GetSharedData()).CustomDelete = FCustomSceneOutlinerDeleteDelegate::CreateRaw(this, &FDataLayerMode::DeleteItems);
@@ -163,7 +201,7 @@ TSharedRef<FSceneOutlinerFilter> FDataLayerMode::CreateHideRuntimeDataLayersFilt
 
 TSharedRef<FSceneOutlinerFilter> FDataLayerMode::CreateHideDataLayerActorsFilter()
 {
-	return MakeShareable(new FDataLayerActorFilter(FDataLayerActorTreeItem::FFilterPredicate::CreateStatic([](const AActor* Actor, const UDataLayerInstance* DataDataLayerInstanceLayer) { return true; }), FSceneOutlinerFilter::EDefaultBehaviour::Pass));
+	return MakeShareable(new FDataLayerActorFilter(FDataLayerActorTreeItem::FFilterPredicate::CreateStatic([](const AActor* Actor, const UDataLayerInstance* DataLayerInstance) { return true; }), FSceneOutlinerFilter::EDefaultBehaviour::Pass));
 }
 
 TSharedRef<FSceneOutlinerFilter> FDataLayerMode::CreateHideUnloadedActorsFilter()
@@ -171,9 +209,34 @@ TSharedRef<FSceneOutlinerFilter> FDataLayerMode::CreateHideUnloadedActorsFilter(
 	return MakeShareable(new FActorDescFilter(FActorDescTreeItem::FFilterPredicate::CreateStatic([](const FWorldPartitionActorDesc* ActorDesc) { return true; }), FSceneOutlinerFilter::EDefaultBehaviour::Pass));
 }
 
+TSharedRef<FSceneOutlinerFilter> FDataLayerMode::CreateHideLevelInstancesFilter()
+{
+	return MakeShareable(new FDataLayerActorFilter(FDataLayerActorTreeItem::FFilterPredicate::CreateStatic([](const AActor* Actor, const UDataLayerInstance* DataLayerInstance)
+	{
+		// Check if actor belongs to a LevelInstance
+		if (const ULevelInstanceSubsystem* LevelInstanceSubsystem = Actor->GetWorld()->GetSubsystem<ULevelInstanceSubsystem>())
+		{
+			if (const ILevelInstanceInterface* ParentLevelInstance = LevelInstanceSubsystem->GetParentLevelInstance(Actor))
+			{
+				if (!LevelInstanceSubsystem->IsEditingLevelInstance(ParentLevelInstance))
+				{
+					return false;
+				}
+			}
+		}
+		// Or if the actor itself is a LevelInstance editor instance
+		return !Actor->IsA<ALevelInstanceEditorInstanceActor>();
+	}), FSceneOutlinerFilter::EDefaultBehaviour::Pass));
+}
+
 int32 FDataLayerMode::GetTypeSortPriority(const ISceneOutlinerTreeItem& Item) const
 {
-	if (Item.IsA<FDataLayerTreeItem>())
+	if (Item.IsA<FWorldDataLayersTreeItem>())
+	{
+		FWorldDataLayersTreeItem* WorldDataLayersTreeItem = (FWorldDataLayersTreeItem*)&Item;
+		return static_cast<int32>(EItemSortOrder::WorldDataLayers) + WorldDataLayersTreeItem->GetSortPriority();
+	}
+	else if (Item.IsA<FDataLayerTreeItem>())
 	{
 		return static_cast<int32>(EItemSortOrder::DataLayer);
 	}
@@ -289,7 +352,7 @@ void FDataLayerMode::OnItemDoubleClick(FSceneOutlinerTreeItemPtr Item)
 	{
 		if (UDataLayerInstance* DataLayerInstance = DataLayerItem->GetDataLayer())
 		{
-			const FScopedDataLayerTransaction Transaction(LOCTEXT("SelectActorsInDataLayer", "Select Actors in Data Layer"), RepresentingWorld.Get());
+			const FScopedDataLayerTransaction Transaction(LOCTEXT("SelectActorsInDataLayer", "Select Actors in Data Layer"), GetOwningWorld());
 			GEditor->SelectNone(/*bNoteSelectionChange*/false, true);
 			DataLayerEditorSubsystem->SelectActorsInDataLayer(DataLayerInstance, /*bSelect*/true, /*bNotify*/true, /*bSelectEvenIfHidden*/true);
 		}
@@ -298,7 +361,7 @@ void FDataLayerMode::OnItemDoubleClick(FSceneOutlinerTreeItemPtr Item)
 	{
 		if (AActor * Actor = DataLayerActorItem->GetActor())
 		{
-			const FScopedDataLayerTransaction Transaction(LOCTEXT("ClickingOnActor", "Clicking on Actor in Data Layer"), RepresentingWorld.Get());
+			const FScopedDataLayerTransaction Transaction(LOCTEXT("ClickingOnActor", "Clicking on Actor in Data Layer"), GetOwningWorld());
 			GEditor->GetSelectedActors()->Modify();
 			GEditor->SelectNone(/*bNoteSelectionChange*/false, true);
 			GEditor->SelectActor(Actor, /*bSelected*/true, /*bNotify*/true, /*bSelectEvenIfHidden*/true);
@@ -338,7 +401,7 @@ void FDataLayerMode::DeleteItems(const TArray<TWeakPtr<ISceneOutlinerTreeItem>>&
 
 	if (!ActorsToRemoveFromDataLayer.IsEmpty())
 	{
-		const FScopedDataLayerTransaction Transaction(LOCTEXT("RemoveActorsFromDataLayer", "Remove Actors from Data Layer"), RepresentingWorld.Get());
+		const FScopedDataLayerTransaction Transaction(LOCTEXT("RemoveActorsFromDataLayer", "Remove Actors from Data Layer"), GetOwningWorld());
 		for (const auto& ItPair : ActorsToRemoveFromDataLayer)
 		{
 			DataLayerEditorSubsystem->RemoveActorsFromDataLayer(ItPair.Value, ItPair.Key);
@@ -353,7 +416,7 @@ void FDataLayerMode::DeleteItems(const TArray<TWeakPtr<ISceneOutlinerTreeItem>>&
 		}
 		
 		{
-			const FScopedDataLayerTransaction Transaction(LOCTEXT("DeleteDataLayers", "Delete Data Layers"), RepresentingWorld.Get());
+			const FScopedDataLayerTransaction Transaction(LOCTEXT("DeleteDataLayers", "Delete Data Layers"), GetOwningWorld());
 			DataLayerEditorSubsystem->DeleteDataLayers(DataLayersToDelete);
 		}
 
@@ -647,7 +710,7 @@ void FDataLayerMode::OnDrop(ISceneOutlinerTreeItem& DropTarget, const FSceneOutl
 			TArray<UDataLayerInstance*> AllSelectedDataLayers = GetSelectedDataLayers(SceneOutliner);
 			if (AllSelectedDataLayers.Num() > 1)
 			{
-				const FScopedDataLayerTransaction Transaction(LOCTEXT("DataLayerOutlinerAddActorsToDataLayers", "Add Actors to Data Layers"), RepresentingWorld.Get());
+				const FScopedDataLayerTransaction Transaction(LOCTEXT("DataLayerOutlinerAddActorsToDataLayers", "Add Actors to Data Layers"), GetOwningWorld());
 				DataLayerEditorSubsystem->AddActorsToDataLayers(ActorsToAdd, AllSelectedDataLayers);
 				return;
 			}
@@ -655,7 +718,7 @@ void FDataLayerMode::OnDrop(ISceneOutlinerTreeItem& DropTarget, const FSceneOutl
 
 		if (TargetDataLayer)
 		{
-			const FScopedDataLayerTransaction Transaction(LOCTEXT("DataLayerOutlinerAddActorsToDataLayer", "Add Actors to Data Layer"), RepresentingWorld.Get());
+			const FScopedDataLayerTransaction Transaction(LOCTEXT("DataLayerOutlinerAddActorsToDataLayer", "Add Actors to Data Layer"), GetOwningWorld());
 			DataLayerEditorSubsystem->AddActorsToDataLayer(ActorsToAdd, TargetDataLayer);
 		}
 	}
@@ -664,7 +727,7 @@ void FDataLayerMode::OnDrop(ISceneOutlinerTreeItem& DropTarget, const FSceneOutl
 		TArray<FDataLayerActorMoveElement> ActorsToMove = GetDataLayerActorPairsFromOperation(Payload.SourceOperation);
 		if (!ActorsToMove.IsEmpty() && TargetDataLayer)
 		{
-			const FScopedDataLayerTransaction Transaction(LOCTEXT("DataLayerOutlinerMoveActorsToDataLayer", "Move Actors to Data Layer"), RepresentingWorld.Get());
+			const FScopedDataLayerTransaction Transaction(LOCTEXT("DataLayerOutlinerMoveActorsToDataLayer", "Move Actors to Data Layer"), GetOwningWorld());
 			for (const auto& DataLayerActorPair : ActorsToMove)
 			{
 				if (AActor* ActorPtr = DataLayerActorPair.Key.Get())
@@ -721,7 +784,7 @@ void FDataLayerMode::SetParentDataLayer(const TArray<UDataLayerInstance*> DataLa
 
 		if (!ValidDataLayers.IsEmpty())
 		{
-			const FScopedDataLayerTransaction Transaction(LOCTEXT("DataLayerOutlinerChangeDataLayersParent", "Change Data Layers Parent"), RepresentingWorld.Get());
+			const FScopedDataLayerTransaction Transaction(LOCTEXT("DataLayerOutlinerChangeDataLayersParent", "Change Data Layers Parent"), GetOwningWorld());
 			for (UDataLayerInstance* DataLayerInstance : ValidDataLayers)
 			{
 				DataLayerEditorSubsystem->SetParentDataLayer(DataLayerInstance, ParentDataLayer);
@@ -892,6 +955,18 @@ void FDataLayerMode::CreateDataLayerPicker(UToolMenu* InMenu, FOnDataLayerPicked
 	Section.AddEntry(FToolMenuEntry::InitWidget("DataLayerPickerWidget", DataLayerPickerWidget, FText::GetEmpty(), false));
 }
 
+UWorld* FDataLayerMode::GetOwningWorld() const
+{
+	UWorld* World = RepresentingWorld.Get();
+	return World ? World->PersistentLevel->GetWorld() : nullptr;
+}
+
+AWorldDataLayers* FDataLayerMode::GetOwningWorldAWorldDataLayers() const
+{
+	UWorld* OwningWorld = GetOwningWorld();
+	return OwningWorld ? OwningWorld->GetWorldDataLayers() : nullptr;
+}
+
 void FDataLayerMode::RegisterContextMenu()
 {
 	UToolMenus* ToolMenus = UToolMenus::Get();
@@ -914,7 +989,7 @@ void FDataLayerMode::RegisterContextMenu()
 
 			bool bHasActorEditorContextDataLayers = false;			
 			TArray<const UDataLayerInstance*> AllDataLayers;
-			if (const UDataLayerSubsystem* DataLayerSubsystem = UWorld::GetSubsystem<UDataLayerSubsystem>(RepresentingWorld.Get()))
+			if (const UDataLayerSubsystem* DataLayerSubsystem = UWorld::GetSubsystem<UDataLayerSubsystem>(GetOwningWorld()))
 			{
 				DataLayerSubsystem->ForEachDataLayer([&AllDataLayers](UDataLayerInstance* DataLayerInstance)
 				{
@@ -928,16 +1003,29 @@ void FDataLayerMode::RegisterContextMenu()
 			{
 				FToolMenuSection& Section = InMenu->AddSection("DataLayers", LOCTEXT("DataLayers", "Data Layers"));
 				
-				auto CreateNewDataLayer = [this, SceneOutliner](UDataLayerInstance* ParentDataLayer = nullptr)
+				auto CreateNewDataLayer = [this, SceneOutliner](UDataLayerInstance* InParentDataLayer = nullptr, const TSet<const UDataLayerAsset*>& InDataLayerAssets = {})
 				{
-					const FScopedDataLayerTransaction Transaction(LOCTEXT("CreateNewDataLayer", "Create New Data Layer"), RepresentingWorld.Get());
-					SelectedDataLayersSet.Empty();
-					SelectedDataLayerActors.Empty();
-					if (UDataLayerAsset* DataLayerAsset = PromptDataLayerAssetSelection())
+					TSet<const UDataLayerAsset*> DataLayerAssets = InDataLayerAssets;
+					DataLayerAssets.Remove(nullptr);
+
+					if (DataLayerAssets.IsEmpty())
+					{
+						if (UDataLayerAsset* DataLayerAsset = PromptDataLayerAssetSelection())
+						{
+							DataLayerAssets.Add(DataLayerAsset);
+						}
+					}
+
+					for (const UDataLayerAsset* DataLayerAsset : DataLayerAssets)
 					{
+						const FScopedDataLayerTransaction Transaction(LOCTEXT("CreateNewDataLayer", "Create New Data Layer"), GetOwningWorld());
+						SelectedDataLayersSet.Empty();
+						SelectedDataLayerActors.Empty();
+
 						FDataLayerCreationParameters CreationParams;
-						CreationParams.DataLayerAsset = DataLayerAsset;
-						CreationParams.ParentDataLayer = ParentDataLayer;
+						CreationParams.DataLayerAsset = const_cast<UDataLayerAsset*>(DataLayerAsset);
+						CreationParams.ParentDataLayer = InParentDataLayer;
+						CreationParams.WorlDataLayers = InParentDataLayer ? InParentDataLayer->GetOuterAWorldDataLayers() : GetOwningWorldAWorldDataLayers();
 						if (UDataLayerInstance* NewDataLayerInstance = DataLayerEditorSubsystem->CreateDataLayerInstance(CreationParams))
 						{
 							SelectedDataLayersSet.Add(NewDataLayerInstance);
@@ -947,6 +1035,52 @@ void FDataLayerMode::RegisterContextMenu()
 					}
 				};
 
+				// If selection contains readonly DataLayerInstances
+				if (Algo::AnyOf(SelectedDataLayersSet, [](const TWeakObjectPtr<const UDataLayerInstance>& DataLayerInstance) { return DataLayerInstance.IsValid() && DataLayerInstance->IsReadOnly(); }))
+				{
+					// Create a menu entry that allows to creating DataLayerInstance in the world using selection referenced DataLayer Assets
+					TSet<const UDataLayerAsset*> ExistingDataLayerAssets;
+					if (AWorldDataLayers* WorldDataLayer = GetOwningWorldAWorldDataLayers())
+					{
+						WorldDataLayer->ForEachDataLayer([&ExistingDataLayerAssets](UDataLayerInstance* DataLayerInstance)
+						{
+							UDataLayerInstanceWithAsset* DataLayerWithAsset = Cast<UDataLayerInstanceWithAsset>(DataLayerInstance);
+							if (const UDataLayerAsset* DataLayerAsset = DataLayerWithAsset ? DataLayerWithAsset->GetAsset() : nullptr)
+							{
+								ExistingDataLayerAssets.Add(DataLayerAsset);
+							}
+							return true;
+						});
+					}
+
+					TSet<const UDataLayerAsset*> DataLayerAssetsToImport;
+					for (const TWeakObjectPtr<const UDataLayerInstance>& DataLayerInstance : SelectedDataLayersSet)
+					{
+						if (DataLayerInstance.IsValid() && DataLayerInstance->IsReadOnly())
+						{
+							const UDataLayerInstanceWithAsset* DataLayerWithAsset = Cast<const UDataLayerInstanceWithAsset>(DataLayerInstance);
+							const UDataLayerAsset* DataLayerAsset = DataLayerWithAsset ? DataLayerWithAsset->GetAsset() : nullptr;
+							if (DataLayerAsset && !ExistingDataLayerAssets.Contains(DataLayerAsset))
+							{
+								DataLayerAssetsToImport.Add(DataLayerAsset);
+							}
+						}
+					}
+
+					if (!DataLayerAssetsToImport.IsEmpty())
+					{
+						Section.AddMenuEntry("ImportDataLayers", LOCTEXT("ImportDataLayers", "Import Data Layer(s)"), FText(), FSlateIcon(),
+							FUIAction(FExecuteAction::CreateLambda([this, DataLayerAssetsToImport, CreateNewDataLayer]()
+								{
+									const FScopedDataLayerTransaction Transaction(LOCTEXT("ImportDataLayers", "Import Data Layers"), GetOwningWorld());
+									CreateNewDataLayer(nullptr, DataLayerAssetsToImport);
+								})));
+					}
+
+					// Readonly selected Data Layer Instances will only show the option to import Data Layers
+					return;
+				}
+				
 				Section.AddMenuEntry("CreateNewDataLayer", LOCTEXT("CreateNewDataLayer", "Create New Data Layer"), FText(), FSlateIcon(),
 					FUIAction(FExecuteAction::CreateLambda([this, CreateNewDataLayer]() { CreateNewDataLayer(); })));
 
@@ -962,9 +1096,10 @@ void FDataLayerMode::RegisterContextMenu()
 						FExecuteAction::CreateLambda([this]() {
 							if (UDataLayerAsset* DataLayerAsset = PromptDataLayerAssetSelection())
 							{
-								const FScopedDataLayerTransaction Transaction(LOCTEXT("AddSelectedActorsToNewDataLayer", "Add Selected Actors to New Data Layer"), RepresentingWorld.Get());
+								const FScopedDataLayerTransaction Transaction(LOCTEXT("AddSelectedActorsToNewDataLayer", "Add Selected Actors to New Data Layer"), GetOwningWorld());
 								FDataLayerCreationParameters CreationParams;
 								CreationParams.DataLayerAsset = DataLayerAsset;
+								CreationParams.WorlDataLayers = GetOwningWorldAWorldDataLayers();
 								if (UDataLayerInstance* NewDataLayerInstance = DataLayerEditorSubsystem->CreateDataLayerInstance(CreationParams))
 								{
 									DataLayerEditorSubsystem->AddSelectedActorsToDataLayer(NewDataLayerInstance);
@@ -978,7 +1113,7 @@ void FDataLayerMode::RegisterContextMenu()
 						FExecuteAction::CreateLambda([this,SelectedDataLayers]() {
 							check(!SelectedDataLayers.IsEmpty());
 							{
-								const FScopedDataLayerTransaction Transaction(LOCTEXT("AddSelectedActorsToSelectedDataLayers", "Add Selected Actors to Selected Data Layers"), RepresentingWorld.Get());
+								const FScopedDataLayerTransaction Transaction(LOCTEXT("AddSelectedActorsToSelectedDataLayers", "Add Selected Actors to Selected Data Layers"), GetOwningWorld());
 								DataLayerEditorSubsystem->AddSelectedActorsToDataLayers(SelectedDataLayers);
 							}}),
 						FCanExecuteAction::CreateLambda([SelectedDataLayers,bSelectedDataLayersContainsLocked] { return !SelectedDataLayers.IsEmpty() && GEditor->GetSelectedActorCount() > 0 && !bSelectedDataLayersContainsLocked; })
@@ -996,7 +1131,7 @@ void FDataLayerMode::RegisterContextMenu()
 								Algo::TransformIf(SelectedDataLayerActors, Actors, [](const FSelectedDataLayerActor& InActor) { return InActor.Value.Get(); }, [](const FSelectedDataLayerActor& InActor) { return const_cast<AActor*>(InActor.Value.Get()); });
 								if (!Actors.IsEmpty())
 								{
-									const FScopedDataLayerTransaction Transaction(LOCTEXT("AddSelectedActorsToDataLayer", "Add Selected Actors to Selected Data Layer"), RepresentingWorld.Get());
+									const FScopedDataLayerTransaction Transaction(LOCTEXT("AddSelectedActorsToDataLayer", "Add Selected Actors to Selected Data Layer"), GetOwningWorld());
 									DataLayerEditorSubsystem->AddActorsToDataLayers(Actors, { TargetDataLayer });
 								}
 							}));
@@ -1029,7 +1164,7 @@ void FDataLayerMode::RegisterContextMenu()
 						FExecuteAction::CreateLambda([this,SelectedDataLayers]() {
 							check(!SelectedDataLayers.IsEmpty());
 							{
-								const FScopedDataLayerTransaction Transaction(LOCTEXT("RemoveSelectedActorsFromSelectedDataLayers_DataLayerMode", "Remove Selected Actors from Selected Data Layers"), RepresentingWorld.Get());
+								const FScopedDataLayerTransaction Transaction(LOCTEXT("RemoveSelectedActorsFromSelectedDataLayers_DataLayerMode", "Remove Selected Actors from Selected Data Layers"), GetOwningWorld());
 								DataLayerEditorSubsystem->RemoveSelectedActorsFromDataLayers(SelectedDataLayers);
 							}}),
 						FCanExecuteAction::CreateLambda([SelectedDataLayers,bSelectedDataLayersContainsLocked] { return !SelectedDataLayers.IsEmpty() && GEditor->GetSelectedActorCount() > 0 && !bSelectedDataLayersContainsLocked; })
@@ -1040,7 +1175,7 @@ void FDataLayerMode::RegisterContextMenu()
 						FExecuteAction::CreateLambda([this,SelectedDataLayers]() {
 							check(!SelectedDataLayers.IsEmpty());
 							{
-								const FScopedDataLayerTransaction Transaction(LOCTEXT("DeleteSelectedDataLayers", "Delete Selected Data Layers"), RepresentingWorld.Get());
+								const FScopedDataLayerTransaction Transaction(LOCTEXT("DeleteSelectedDataLayers", "Delete Selected Data Layers"), GetOwningWorld());
 								DataLayerEditorSubsystem->DeleteDataLayers(SelectedDataLayers);
 							}}),
 						FCanExecuteAction::CreateLambda([SelectedDataLayers,bSelectedDataLayersContainsLocked] { return !SelectedDataLayers.IsEmpty() && !bSelectedDataLayersContainsLocked; })
@@ -1072,7 +1207,7 @@ void FDataLayerMode::RegisterContextMenu()
 						FExecuteAction::CreateLambda([this,SelectedDataLayers]() {
 							check(!SelectedDataLayers.IsEmpty());
 							{
-								const FScopedDataLayerTransaction Transaction(LOCTEXT("SelectActorsInDataLayers", "Select Actors in Data Layers"), RepresentingWorld.Get());
+								const FScopedDataLayerTransaction Transaction(LOCTEXT("SelectActorsInDataLayers", "Select Actors in Data Layers"), GetOwningWorld());
 								GEditor->SelectNone(/*bNoteSelectionChange*/false, /*bDeselectBSPSurfs*/true);
 								DataLayerEditorSubsystem->SelectActorsInDataLayers(SelectedDataLayers, /*bSelect*/true, /*bNotify*/true, /*bSelectEvenIfHidden*/true);
 							}}),
@@ -1084,7 +1219,7 @@ void FDataLayerMode::RegisterContextMenu()
 						FExecuteAction::CreateLambda([this,SelectedDataLayers]() {
 							check(!SelectedDataLayers.IsEmpty());
 							{
-								const FScopedDataLayerTransaction Transaction(LOCTEXT("AppendActorsToSelection", "Append Actors in Data Layer to Selection"), RepresentingWorld.Get());
+								const FScopedDataLayerTransaction Transaction(LOCTEXT("AppendActorsToSelection", "Append Actors in Data Layer to Selection"), GetOwningWorld());
 								DataLayerEditorSubsystem->SelectActorsInDataLayers(SelectedDataLayers, /*bSelect*/true, /*bNotify*/true, /*bSelectEvenIfHidden*/true);
 							}}),
 						FCanExecuteAction::CreateLambda([SelectedDataLayers,bSelectedDataLayersContainsLocked] { return !SelectedDataLayers.IsEmpty() && !bSelectedDataLayersContainsLocked; })
@@ -1095,7 +1230,7 @@ void FDataLayerMode::RegisterContextMenu()
 						FExecuteAction::CreateLambda([this,SelectedDataLayers]() {
 							check(!SelectedDataLayers.IsEmpty());
 							{
-								const FScopedDataLayerTransaction Transaction(LOCTEXT("DeselectActors", "Deselect Actors in Data Layer"), RepresentingWorld.Get());
+								const FScopedDataLayerTransaction Transaction(LOCTEXT("DeselectActors", "Deselect Actors in Data Layer"), GetOwningWorld());
 								DataLayerEditorSubsystem->SelectActorsInDataLayers(SelectedDataLayers, /*bSelect*/false, /*bNotifySelectActors*/true);
 							}}),
 						FCanExecuteAction::CreateLambda([SelectedDataLayers] { return !SelectedDataLayers.IsEmpty(); })
@@ -1135,7 +1270,7 @@ void FDataLayerMode::RegisterContextMenu()
 					FUIAction(
 						FExecuteAction::CreateLambda([this, SceneOutliner, SelectedDataLayers]() 
 						{
-							const FScopedDataLayerTransaction Transaction(LOCTEXT("MakeCurrentDataLayers", "Make Current Data Layer(s)"), RepresentingWorld.Get());
+							const FScopedDataLayerTransaction Transaction(LOCTEXT("MakeCurrentDataLayers", "Make Current Data Layer(s)"), GetOwningWorld());
 							for (TWeakObjectPtr<const UDataLayerInstance>& DataLayer : SelectedDataLayersSet)
 							{
 								if (DataLayer.IsValid() && !DataLayer->IsLocked())
@@ -1151,7 +1286,7 @@ void FDataLayerMode::RegisterContextMenu()
 					FUIAction(
 						FExecuteAction::CreateLambda([this, SceneOutliner, SelectedDataLayers]()
 						{
-							const FScopedDataLayerTransaction Transaction(LOCTEXT("RemoveCurrentDataLayers", "Remove Current Data Layer(s)"), RepresentingWorld.Get());
+							const FScopedDataLayerTransaction Transaction(LOCTEXT("RemoveCurrentDataLayers", "Remove Current Data Layer(s)"), GetOwningWorld());
 							for (TWeakObjectPtr<const UDataLayerInstance>& DataLayer : SelectedDataLayersSet)
 							{
 								if (DataLayer.IsValid() && !DataLayer->IsLocked())
@@ -1169,7 +1304,7 @@ void FDataLayerMode::RegisterContextMenu()
 						{
 							check(!AllDataLayers.IsEmpty());
 							{
-								const FScopedDataLayerTransaction Transaction(LOCTEXT("ClearCurrentDataLayers", "Clear Current Data Layers"), RepresentingWorld.Get());
+								const FScopedDataLayerTransaction Transaction(LOCTEXT("ClearCurrentDataLayers", "Clear Current Data Layers"), GetOwningWorld());
 								for (const UDataLayerInstance* DataLayer : AllDataLayers)
 								{
 									DataLayerEditorSubsystem->RemoveFromActorEditorContext(const_cast<UDataLayerInstance*>(DataLayer));
@@ -1187,7 +1322,7 @@ void FDataLayerMode::RegisterContextMenu()
 						FExecuteAction::CreateLambda([this,AllDataLayers]() {
 							check(!AllDataLayers.IsEmpty());
 							{
-								const FScopedDataLayerTransaction Transaction(LOCTEXT("MakeAllDataLayersVisible", "Make All Data Layers Visible"), RepresentingWorld.Get());
+								const FScopedDataLayerTransaction Transaction(LOCTEXT("MakeAllDataLayersVisible", "Make All Data Layers Visible"), GetOwningWorld());
 								DataLayerEditorSubsystem->MakeAllDataLayersVisible();
 							}}),
 						FCanExecuteAction::CreateLambda([AllDataLayers] { return !AllDataLayers.IsEmpty(); })
@@ -1276,7 +1411,7 @@ void FDataLayerMode::CreateViewContent(FMenuBuilder& MenuBuilder)
 			{
 				if (AWorldDataLayers* WorldDataLayers = RepresentingWorld.IsValid() ? RepresentingWorld->GetWorldDataLayers() : nullptr)
 				{
-					const FScopedDataLayerTransaction Transaction(LOCTEXT("ToggleAllowRuntimeDataLayerEditingTransaction", "Toggle Allow Runtime Data Layer Editing"), RepresentingWorld.Get());
+					const FScopedDataLayerTransaction Transaction(LOCTEXT("ToggleAllowRuntimeDataLayerEditingTransaction", "Toggle Allow Runtime Data Layer Editing"), GetOwningWorld());
 					WorldDataLayers->SetAllowRuntimeDataLayerEditing(!WorldDataLayers->GetAllowRuntimeDataLayerEditing());
 				}
 				SceneOutliner->FullRefresh();
@@ -1293,7 +1428,7 @@ void FDataLayerMode::CreateViewContent(FMenuBuilder& MenuBuilder)
 	);
 
 	TArray<UDataLayerInstance*> AllDataLayers;
-	if (const UDataLayerSubsystem* DataLayerSubsystem = UWorld::GetSubsystem<UDataLayerSubsystem>(RepresentingWorld.Get()))
+	if (const UDataLayerSubsystem* DataLayerSubsystem = UWorld::GetSubsystem<UDataLayerSubsystem>(GetOwningWorld()))
 	{
 		DataLayerSubsystem->ForEachDataLayer([&AllDataLayers](UDataLayerInstance* DataLayer)
 		{
@@ -1309,7 +1444,7 @@ void FDataLayerMode::CreateViewContent(FMenuBuilder& MenuBuilder)
 		FUIAction(
 			FExecuteAction::CreateLambda([this]()
 			{
-				const FScopedDataLayerTransaction Transaction(LOCTEXT("ResetDataLayerUserSettings", "Reset User Settings"), RepresentingWorld.Get());
+				const FScopedDataLayerTransaction Transaction(LOCTEXT("ResetDataLayerUserSettings", "Reset User Settings"), GetOwningWorld());
 				DataLayerEditorSubsystem->ResetUserSettings();
 			})
 		)
@@ -1386,6 +1521,7 @@ TUniquePtr<ISceneOutlinerHierarchy> FDataLayerMode::CreateHierarchy()
 	DataLayerHierarchy->SetShowUnloadedActors(!bHideUnloadedActors);
 	DataLayerHierarchy->SetShowOnlySelectedActors(bShowOnlySelectedActors);
 	DataLayerHierarchy->SetHighlightSelectedDataLayers(bHighlightSelectedDataLayers);
+	DataLayerHierarchy->SetShowLevelInstanceContent(!bHideLevelInstanceContent);
 	return DataLayerHierarchy;
 }
 
@@ -1494,6 +1630,16 @@ void FDataLayerMode::ChooseRepresentingWorld()
 			else if (Context.WorldType == EWorldType::Editor)
 			{
 				RepresentingWorld = Context.World();
+				if (RepresentingWorld.IsValid())
+				{
+					// Favor Current Level's outer world if exists
+					ULevel* CurrentLevel = RepresentingWorld->GetCurrentLevel();
+					if (CurrentLevel && !CurrentLevel->IsPersistentLevel() && CurrentLevel->GetWorldDataLayers())
+					{
+						RepresentingWorld = CurrentLevel->GetTypedOuter<UWorld>();
+						check(RepresentingWorld.IsValid());
+					}
+				}
 			}
 		}
 	}
diff --git a/Engine/Source/Editor/DataLayerEditor/Private/DataLayer/DataLayerMode.h b/Engine/Source/Editor/DataLayerEditor/Private/DataLayer/DataLayerMode.h
--- a/Engine/Source/Editor/DataLayerEditor/Private/DataLayer/DataLayerMode.h
+++ b/Engine/Source/Editor/DataLayerEditor/Private/DataLayer/DataLayerMode.h
@@ -11,6 +11,7 @@
 class UDataLayerEditorSubsystem;
 class SDataLayerBrowser;
 class UWorld;
+class AWorldDataLayers;
 
 struct FDataLayerModeParams
 {
@@ -31,9 +32,10 @@ class FDataLayerMode : public ISceneOutlinerMode
 public:
 	enum class EItemSortOrder : int32
 	{
-		DataLayer = 0,
-		Actor = 10,
-		Unloaded = 20,
+		WorldDataLayers = 0,
+		DataLayer = 10,
+		Actor = 20,
+		Unloaded = 30,
 	};
 
 	FDataLayerMode(const FDataLayerModeParams& Params);
@@ -90,6 +92,9 @@ protected:
 	bool bShowOnlySelectedActors;
 	/** Should highlight DataLayers containing selected actors */
 	bool bHighlightSelectedDataLayers;
+	/** Should level instance actors content be hidden */
+	bool bHideLevelInstanceContent;
+
 	/** Delegate to call when an item is picked */
 	FOnSceneOutlinerItemPicked OnItemPicked;
 
@@ -110,6 +115,8 @@ private:
 	bool ShouldExpandDataLayer(const UDataLayerInstance* DataLayer) const;
 	bool ContainsSelectedChildDataLayer(const UDataLayerInstance* DataLayer) const;
 	void RefreshSelection();
+	UWorld* GetOwningWorld() const;
+	AWorldDataLayers* GetOwningWorldAWorldDataLayers() const;
 	FSceneOutlinerDragValidationInfo ValidateDrop(const ISceneOutlinerTreeItem& DropTarget, bool bMoveOperation = false) const;
 
 	UDataLayerAsset* PromptDataLayerAssetSelection();
@@ -120,6 +127,7 @@ private:
 	static TSharedRef<FSceneOutlinerFilter> CreateHideRuntimeDataLayersFilter();
 	static TSharedRef<FSceneOutlinerFilter> CreateHideDataLayerActorsFilter();
 	static TSharedRef<FSceneOutlinerFilter> CreateHideUnloadedActorsFilter();
+	static TSharedRef<FSceneOutlinerFilter> CreateHideLevelInstancesFilter();
 
 	/** The world which we are currently representing */
 	TWeakObjectPtr<UWorld> RepresentingWorld;
diff --git a/Engine/Source/Editor/DataLayerEditor/Private/DataLayer/DataLayerOutlinerDeleteButtonColumn.cpp b/Engine/Source/Editor/DataLayerEditor/Private/DataLayer/DataLayerOutlinerDeleteButtonColumn.cpp
--- a/Engine/Source/Editor/DataLayerEditor/Private/DataLayer/DataLayerOutlinerDeleteButtonColumn.cpp
+++ b/Engine/Source/Editor/DataLayerEditor/Private/DataLayer/DataLayerOutlinerDeleteButtonColumn.cpp
@@ -44,7 +44,7 @@ const TSharedRef<SWidget> FDataLayerOutlinerDeleteButtonColumn::ConstructRowWidg
 			{
 				AActor* Actor = DataLayerActorItem->GetActor();
 				const UDataLayerInstance* DataLayer = DataLayerActorItem->GetDataLayer();
-				return (Actor && DataLayer && !DataLayer->IsLocked()) ? EVisibility::Visible : EVisibility::Collapsed;
+				return (Actor && DataLayer && !DataLayer->IsLocked() && TreeItem->CanInteract()) ? EVisibility::Visible : EVisibility::Collapsed;
 			})
 			.OnClicked_Lambda([this, TreeItem, DataLayerActorItem]()
 			{
diff --git a/Engine/Source/Editor/DataLayerEditor/Private/DataLayer/DataLayerOutlinerIsLoadedInEditorColumn.cpp b/Engine/Source/Editor/DataLayerEditor/Private/DataLayer/DataLayerOutlinerIsLoadedInEditorColumn.cpp
--- a/Engine/Source/Editor/DataLayerEditor/Private/DataLayer/DataLayerOutlinerIsLoadedInEditorColumn.cpp
+++ b/Engine/Source/Editor/DataLayerEditor/Private/DataLayer/DataLayerOutlinerIsLoadedInEditorColumn.cpp
@@ -46,7 +46,7 @@ const TSharedRef<SWidget> FDataLayerOutlinerIsLoadedInEditorColumn::ConstructRow
 			.VAlign(VAlign_Center)
 			[
 				SNew(SCheckBox)
-				.IsEnabled_Lambda([this, TreeItem]()
+				.IsEnabled_Lambda([TreeItem]()
 				{
 					FDataLayerTreeItem* DataLayerTreeItem = TreeItem->CastTo<FDataLayerTreeItem>();
 					const UDataLayerInstance* DataLayer = DataLayerTreeItem->GetDataLayer();
@@ -54,7 +54,13 @@ const TSharedRef<SWidget> FDataLayerOutlinerIsLoadedInEditorColumn::ConstructRow
 					const bool bIsParentLoaded = ParentDataLayer ? ParentDataLayer->IsEffectiveLoadedInEditor() : true;
 					return bIsParentLoaded && DataLayer && DataLayer->GetWorld() && !DataLayer->GetWorld()->IsPlayInEditor();
 				})
-				.IsChecked_Lambda([this, TreeItem]()
+				.Visibility_Lambda([TreeItem]()
+				{
+					FDataLayerTreeItem* DataLayerTreeItem = TreeItem->CastTo<FDataLayerTreeItem>();
+					const UDataLayerInstance* DataLayerInstance = DataLayerTreeItem->GetDataLayer();
+					return DataLayerInstance && !DataLayerInstance->IsReadOnly() ? EVisibility::Visible : EVisibility::Collapsed;
+				})
+				.IsChecked_Lambda([TreeItem]()
 				{
 					FDataLayerTreeItem* DataLayerTreeItem = TreeItem->CastTo<FDataLayerTreeItem>();
 					UDataLayerInstance* DataLayer = DataLayerTreeItem->GetDataLayer();
diff --git a/Engine/Source/Editor/DataLayerEditor/Private/DataLayer/DataLayerTreeItem.cpp b/Engine/Source/Editor/DataLayerEditor/Private/DataLayer/DataLayerTreeItem.cpp
--- a/Engine/Source/Editor/DataLayerEditor/Private/DataLayer/DataLayerTreeItem.cpp
+++ b/Engine/Source/Editor/DataLayerEditor/Private/DataLayer/DataLayerTreeItem.cpp
@@ -22,14 +22,20 @@ FDataLayerTreeItem::FDataLayerTreeItem(UDataLayerInstance* InDataLayerInstance)
 
 FString FDataLayerTreeItem::GetDisplayString() const
 {
-	const UDataLayerInstance* DataLayerPtr = DataLayerInstance.Get();
-	return DataLayerPtr ? DataLayerPtr->GetDataLayerShortName() : LOCTEXT("DataLayerForMissingDataLayer", "(Deleted Data Layer)").ToString();
+	const UDataLayerInstance* DataLayerInstancePtr = DataLayerInstance.Get();
+	return DataLayerInstancePtr ? DataLayerInstancePtr->GetDataLayerShortName() : LOCTEXT("DataLayerForMissingDataLayer", "(Deleted Data Layer)").ToString();
 }
 
 bool FDataLayerTreeItem::GetVisibility() const
 {
-	const UDataLayerInstance* DataLayerPtr = DataLayerInstance.Get();
-	return DataLayerPtr && DataLayerPtr->IsVisible();
+	const UDataLayerInstance* DataLayerInstancePtr = DataLayerInstance.Get();
+	return DataLayerInstancePtr && DataLayerInstancePtr->IsVisible();
+}
+
+bool FDataLayerTreeItem::ShouldShowVisibilityState() const
+{
+	const UDataLayerInstance* DataLayerInstancePtr = DataLayerInstance.Get();
+	return DataLayerInstancePtr && !DataLayerInstancePtr->IsReadOnly();
 }
 
 bool FDataLayerTreeItem::CanInteract() const 
@@ -44,9 +50,9 @@ TSharedRef<SWidget> FDataLayerTreeItem::GenerateLabelWidget(ISceneOutliner& Outl
 
 void FDataLayerTreeItem::OnVisibilityChanged(const bool bNewVisibility)
 {
-	if (UDataLayerInstance* DataLayerPtr = DataLayerInstance.Get())
+	if (UDataLayerInstance* DataLayerInstancePtr = DataLayerInstance.Get())
 	{
-		UDataLayerEditorSubsystem::Get()->SetDataLayerVisibility(DataLayerPtr, bNewVisibility);
+		UDataLayerEditorSubsystem::Get()->SetDataLayerVisibility(DataLayerInstancePtr, bNewVisibility);
 	}
 }
 
@@ -54,9 +60,9 @@ bool FDataLayerTreeItem::ShouldBeHighlighted() const
 {
 	if (bIsHighlighedtIfSelected)
 	{
-		if (UDataLayerInstance* DataLayerPtr = DataLayerInstance.Get())
+		if (UDataLayerInstance* DataLayerInstancePtr = DataLayerInstance.Get())
 		{
-			return UDataLayerEditorSubsystem::Get()->DoesDataLayerContainSelectedActors(DataLayerPtr);
+			return UDataLayerEditorSubsystem::Get()->DoesDataLayerContainSelectedActors(DataLayerInstancePtr);
 		}
 	}
 	return false;
diff --git a/Engine/Source/Editor/DataLayerEditor/Private/DataLayer/DataLayerTreeItem.h b/Engine/Source/Editor/DataLayerEditor/Private/DataLayer/DataLayerTreeItem.h
--- a/Engine/Source/Editor/DataLayerEditor/Private/DataLayer/DataLayerTreeItem.h
+++ b/Engine/Source/Editor/DataLayerEditor/Private/DataLayer/DataLayerTreeItem.h
@@ -24,7 +24,7 @@ public:
 	virtual bool HasVisibilityInfo() const override { return true; }
 	virtual void OnVisibilityChanged(const bool bNewVisibility) override;
 	virtual bool GetVisibility() const override;
-	virtual bool ShouldShowVisibilityState() const { return true; }
+	virtual bool ShouldShowVisibilityState() const override;
 	/* End ISceneOutlinerTreeItem Implementation */
 
 	bool ShouldBeHighlighted() const;
diff --git a/Engine/Source/Editor/DataLayerEditor/Private/DataLayer/DataLayersActorDescTreeItem.h b/Engine/Source/Editor/DataLayerEditor/Private/DataLayer/DataLayersActorDescTreeItem.h
--- a/Engine/Source/Editor/DataLayerEditor/Private/DataLayer/DataLayersActorDescTreeItem.h
+++ b/Engine/Source/Editor/DataLayerEditor/Private/DataLayer/DataLayersActorDescTreeItem.h
@@ -6,6 +6,7 @@
 #include "SceneOutlinerFwd.h"
 #include "ActorDescTreeItem.h"
 #include "WorldPartition/DataLayer/DataLayerInstance.h"
+#include "LevelInstance/LevelInstanceSubsystem.h"
 
 //////////////////////////////////////////////////////////////////////////
 // FDataLayerActorTreeItemData
@@ -37,15 +38,44 @@ public:
 	FDataLayerActorDescTreeItem(const FDataLayerActorDescTreeItemData& InData)
 		: FActorDescTreeItem(InData.ActorGuid, InData.Container)
 		, DataLayer(InData.DataLayer)
-		, IDDataLayerActorDesc(FDataLayerActorDescTreeItem::ComputeTreeItemID(InData.ActorGuid, DataLayer.Get()))
+		, IDDataLayerActorDesc(FDataLayerActorDescTreeItem::ComputeTreeItemID(InData.ActorGuid, DataLayer.Get(), FDataLayerActorDescTreeItem::GetParentActors(InData.Container)))
 	{
+		if (ActorDescHandle.IsValid())
+		{
+			UActorDescContainer* Container = ActorDescHandle->GetContainer();
+			UWorld* OwningWorld = Container->GetWorld();
+			ULevelInstanceSubsystem* LevelInstanceSubsystem = UWorld::GetSubsystem<ULevelInstanceSubsystem>(OwningWorld);
+			ULevel* Level = Container->GetTypedOuter<UWorld>()->PersistentLevel;
+			if (LevelInstanceSubsystem && (Level != OwningWorld->GetCurrentLevel()))
+			{
+				DisplayString = LevelInstanceSubsystem->PrefixWithParentLevelInstanceActorLabels(DisplayString, Level);
+			}
+		}
 	}
 
 	UDataLayerInstance* GetDataLayer() const { return DataLayer.Get(); }
 	
-	static uint32 ComputeTreeItemID(const FGuid& InActorGuid, const UDataLayerInstance* InDataLayer)
+	static uint32 ComputeTreeItemID(const FGuid& InActorGuid, const UDataLayerInstance* InDataLayer, const TArray<AActor*>& InParentActors)
 	{
-		return HashCombine(GetTypeHash(InActorGuid), GetTypeHash(FObjectKey(InDataLayer)));
+		uint32 ID = HashCombine(GetTypeHash(InActorGuid), GetTypeHash(FObjectKey(InDataLayer)));
+		for (AActor* ParentActor : InParentActors)
+		{
+			ID = HashCombine(ID, GetTypeHash(ParentActor->GetActorGuid()));
+		}
+		return ID;
+	}
+
+	static TArray<AActor*> GetParentActors(UActorDescContainer* InContainer)
+	{
+		if (InContainer)
+		{
+			if (ULevelInstanceSubsystem* LevelInstanceSubsystem = UWorld::GetSubsystem<ULevelInstanceSubsystem>(InContainer->GetWorld()))
+			{
+				ULevel* Level = InContainer->GetTypedOuter<UWorld>()->PersistentLevel;
+				return LevelInstanceSubsystem->GetParentLevelInstanceActors(Level);
+			}
+		}
+		return TArray<AActor*>();
 	}
 
 	bool Filter(FFilterPredicate Pred) const
diff --git a/Engine/Source/Editor/DataLayerEditor/Private/DataLayer/SDataLayerBrowser.cpp b/Engine/Source/Editor/DataLayerEditor/Private/DataLayer/SDataLayerBrowser.cpp
--- a/Engine/Source/Editor/DataLayerEditor/Private/DataLayer/SDataLayerBrowser.cpp
+++ b/Engine/Source/Editor/DataLayerEditor/Private/DataLayer/SDataLayerBrowser.cpp
@@ -104,6 +104,7 @@ void SDataLayerBrowser::Construct(const FArguments& InArgs)
 	InitOptions.bShowHeaderRow = true;
 	InitOptions.bShowParentTree = true;
 	InitOptions.bShowCreateNewFolder = false;
+	InitOptions.bShowTransient = true;
 	InitOptions.ModeFactory = FCreateSceneOutlinerMode::CreateLambda([this](SSceneOutliner* Outliner) { return new FDataLayerMode(FDataLayerModeParams(Outliner, this, nullptr)); });
 	InitOptions.ColumnMap.Add(FDataLayerOutlinerIsVisibleColumn::GetID(), FSceneOutlinerColumnInfo(ESceneOutlinerColumnVisibility::Visible, 0, FCreateSceneOutlinerColumn::CreateLambda([](ISceneOutliner& InSceneOutliner) { return MakeShareable(new FDataLayerOutlinerIsVisibleColumn(InSceneOutliner)); })));
 	InitOptions.ColumnMap.Add(FDataLayerOutlinerIsLoadedInEditorColumn::GetID(), FSceneOutlinerColumnInfo(ESceneOutlinerColumnVisibility::Visible, 1, FCreateSceneOutlinerColumn::CreateLambda([](ISceneOutliner& InSceneOutliner) { return MakeShareable(new FDataLayerOutlinerIsLoadedInEditorColumn(InSceneOutliner)); })));
diff --git a/Engine/Source/Editor/DataLayerEditor/Private/DataLayer/SDataLayerTreeLabel.cpp b/Engine/Source/Editor/DataLayerEditor/Private/DataLayer/SDataLayerTreeLabel.cpp
--- a/Engine/Source/Editor/DataLayerEditor/Private/DataLayer/SDataLayerTreeLabel.cpp
+++ b/Engine/Source/Editor/DataLayerEditor/Private/DataLayer/SDataLayerTreeLabel.cpp
@@ -93,7 +93,7 @@ void SDataLayerTreeLabel::Construct(const FArguments& InArgs, FDataLayerTreeItem
 		.VAlign(VAlign_Center)
 		[
 			SNew(SImage)
-			.Visibility_Lambda([this] { return (DataLayerPtr.IsValid() && DataLayerPtr->IsLocked() && !DataLayerPtr->GetWorld()->IsPlayInEditor()) ? EVisibility::Visible : EVisibility::Collapsed; })
+			.Visibility_Lambda([this] { return (DataLayerPtr.IsValid() && DataLayerPtr->IsLocked() && !DataLayerPtr->IsReadOnly() && !DataLayerPtr->GetWorld()->IsPlayInEditor()) ? EVisibility::Visible : EVisibility::Collapsed; })
 			.ColorAndOpacity(FSlateColor::UseForeground())
 			.Image(FAppStyle::GetBrush(TEXT("PropertyWindow.Locked")))
 			.ToolTipText(LOCTEXT("LockedRuntimeDataLayerEditing", "Locked editing. (To allow editing, in Data Layer Outliner, go to Advanced -> Allow Runtime Data Layer Editing)"))
diff --git a/Engine/Source/Editor/DataLayerEditor/Private/DataLayer/WorldDataLayersTreeItem.cpp b/Engine/Source/Editor/DataLayerEditor/Private/DataLayer/WorldDataLayersTreeItem.cpp
new file mode 100644
--- /dev/null
+++ b/Engine/Source/Editor/DataLayerEditor/Private/DataLayer/WorldDataLayersTreeItem.cpp
@@ -0,0 +1,118 @@
+// Copyright Epic Games, Inc. All Rights Reserved.
+
+#include "WorldDataLayersTreeItem.h"
+#include "WorldPartition/DataLayer/WorldDataLayers.h"
+#include "ISceneOutliner.h"
+#include "WorldTreeItem.h"
+#include "SceneOutlinerPublicTypes.h"
+#include "Styling/SlateIconFinder.h"
+#include "Widgets/Images/SImage.h"
+
+#define LOCTEXT_NAMESPACE "DataLayer"
+
+const FSceneOutlinerTreeItemType FWorldDataLayersTreeItem::Type(&ISceneOutlinerTreeItem::Type);
+
+struct SWorldDataLayersTreeLabel : FSceneOutlinerCommonLabelData, public SCompoundWidget
+{
+	SLATE_BEGIN_ARGS(SWorldDataLayersTreeLabel) {}
+	SLATE_END_ARGS()
+
+	void Construct(const FArguments& InArgs, FWorldDataLayersTreeItem& WorldDataLayersItem, ISceneOutliner& SceneOutliner, const STableRow<FSceneOutlinerTreeItemPtr>& InRow)
+	{
+		TreeItemPtr = StaticCastSharedRef<FWorldDataLayersTreeItem>(WorldDataLayersItem.AsShared());
+		WeakSceneOutliner = StaticCastSharedRef<ISceneOutliner>(SceneOutliner.AsShared());
+
+		ChildSlot
+		[
+			SNew(SHorizontalBox)
+
+			+ SHorizontalBox::Slot()
+			.AutoWidth()
+			.VAlign(VAlign_Center)
+			.Padding(FSceneOutlinerDefaultTreeItemMetrics::IconPadding())
+			[
+				SNew(SBox)
+				.WidthOverride(FSceneOutlinerDefaultTreeItemMetrics::IconSize())
+				.HeightOverride(FSceneOutlinerDefaultTreeItemMetrics::IconSize())
+				[
+					SNew(SImage)
+					.Image(FSlateIconFinder::FindIconBrushForClass(UWorld::StaticClass()))
+					.ColorAndOpacity(FSlateColor::UseForeground())
+					.ToolTipText(LOCTEXT("LevelIcon_Tooltip", "Level"))
+				]
+			]
+
+			+ SHorizontalBox::Slot()
+			.FillWidth(1.0f)
+			.VAlign(VAlign_Center)
+			.Padding(0.0f, 2.0f)
+			[
+				SNew(STextBlock)
+				.Text(this, &SWorldDataLayersTreeLabel::GetDisplayText)
+				.HighlightText(SceneOutliner.GetFilterHighlightText())
+				.ColorAndOpacity(this, &SWorldDataLayersTreeLabel::GetForegroundColor)
+			]
+		];
+	}
+
+private:
+	TWeakPtr<FWorldDataLayersTreeItem> TreeItemPtr;
+
+	FText GetDisplayText() const
+	{
+		auto Item = TreeItemPtr.Pin();
+		return Item.IsValid() ? FText::FromString(Item->GetDisplayString()) : FText();
+	}
+
+	FSlateColor GetForegroundColor() const
+	{
+		if (auto BaseColor = FSceneOutlinerCommonLabelData::GetForegroundColor(*TreeItemPtr.Pin()))
+		{
+			return BaseColor.GetValue();
+		}
+
+		return FSlateColor::UseForeground();
+	}
+};
+
+FWorldDataLayersTreeItem::FWorldDataLayersTreeItem(AWorldDataLayers* InWorldDataLayers)
+	: ISceneOutlinerTreeItem(Type)
+	, WorldDataLayers(InWorldDataLayers)
+	, ID(InWorldDataLayers)
+{
+	Flags.bIsExpanded = true;
+}
+
+FString FWorldDataLayersTreeItem::GetDisplayString() const
+{
+	if (const AWorldDataLayers* WorldDataLayersPtr = WorldDataLayers.Get())
+	{
+		check(WorldDataLayersPtr->GetLevel());
+		return SceneOutliner::GetWorldDescription(WorldDataLayersPtr->GetLevel()->GetTypedOuter<UWorld>()).ToString();
+	}
+	return LOCTEXT("UnknownWorldDataLayers", "Unknown").ToString();
+}
+
+bool FWorldDataLayersTreeItem::CanInteract() const
+{
+	return !IsReadOnly();
+}
+
+bool FWorldDataLayersTreeItem::IsReadOnly() const
+{
+	const AWorldDataLayers* WorldDataLayersPtr = WorldDataLayers.Get();
+	UWorld* WorldPtr = WorldDataLayersPtr ? WorldDataLayersPtr->GetLevel()->GetTypedOuter<UWorld>() : nullptr;
+	return (WorldDataLayersPtr && !WorldDataLayersPtr->IsMainWorldDataLayers()) || !WorldPtr || (WorldPtr->WorldType != EWorldType::Editor);
+}
+
+TSharedRef<SWidget> FWorldDataLayersTreeItem::GenerateLabelWidget(ISceneOutliner& Outliner, const STableRow<FSceneOutlinerTreeItemPtr>& InRow)
+{
+	return SNew(SWorldDataLayersTreeLabel, *this, Outliner, InRow);
+}
+
+int32 FWorldDataLayersTreeItem::GetSortPriority() const
+{
+	return CanInteract() ? 0 : 1;
+}
+
+#undef LOCTEXT_NAMESPACE
\ No newline at end of file
diff --git a/Engine/Source/Editor/DataLayerEditor/Private/DataLayer/WorldDataLayersTreeItem.h b/Engine/Source/Editor/DataLayerEditor/Private/DataLayer/WorldDataLayersTreeItem.h
new file mode 100644
--- /dev/null
+++ b/Engine/Source/Editor/DataLayerEditor/Private/DataLayer/WorldDataLayersTreeItem.h
@@ -0,0 +1,48 @@
+// Copyright Epic Games, Inc. All Rights Reserved.
+
+#pragma once
+
+#include "CoreMinimal.h"
+#include "Misc/Attribute.h"
+#include "SceneOutlinerFwd.h"
+#include "ISceneOutlinerTreeItem.h"
+
+class AWorldDataLayers;
+
+struct FWorldDataLayersTreeItem : ISceneOutlinerTreeItem
+{
+public:
+	FWorldDataLayersTreeItem(AWorldDataLayers* InWorldDataLayers);
+	AWorldDataLayers* GetWorldDataLayers() const { return WorldDataLayers.Get(); }
+
+	/* Begin ISceneOutlinerTreeItem Implementation */
+	virtual bool IsValid() const override { return WorldDataLayers.IsValid(); }
+	virtual FSceneOutlinerTreeItemID GetID() const override { return ID; }
+	virtual FString GetDisplayString() const override;
+	virtual bool CanInteract() const override;
+	virtual TSharedRef<SWidget> GenerateLabelWidget(ISceneOutliner& Outliner, const STableRow<FSceneOutlinerTreeItemPtr>& InRow) override;
+	virtual bool ShouldShowVisibilityState() const override { return false; }
+	/* End ISceneOutlinerTreeItem Implementation */
+
+	int32 GetSortPriority() const;
+
+	static const FSceneOutlinerTreeItemType Type;
+
+	DECLARE_DELEGATE_RetVal_OneParam(bool, FFilterPredicate, const AWorldDataLayers*);
+	DECLARE_DELEGATE_RetVal_OneParam(bool, FInteractivePredicate, const AWorldDataLayers*);
+
+	bool Filter(FFilterPredicate Pred) const
+	{
+		return Pred.Execute(GetWorldDataLayers());
+	}
+
+	bool GetInteractiveState(FInteractivePredicate Pred) const
+	{
+		return Pred.Execute(GetWorldDataLayers());
+	}
+
+private:
+	bool IsReadOnly() const;
+	TWeakObjectPtr<AWorldDataLayers> WorldDataLayers;
+	const FObjectKey ID;
+};
\ No newline at end of file
diff --git a/Engine/Source/Editor/DataLayerEditor/Public/DataLayer/DataLayerEditorSubsystem.h b/Engine/Source/Editor/DataLayerEditor/Public/DataLayer/DataLayerEditorSubsystem.h
--- a/Engine/Source/Editor/DataLayerEditor/Public/DataLayer/DataLayerEditorSubsystem.h
+++ b/Engine/Source/Editor/DataLayerEditor/Public/DataLayer/DataLayerEditorSubsystem.h
@@ -16,6 +16,7 @@ class FLevelEditorViewportClient;
 class UEditorEngine;
 class ULevel;
 class UWorld;
+class UWorldPartition;
 template<typename TItemType> class IFilter;
 
 USTRUCT(BlueprintType)
@@ -651,6 +652,11 @@ private:
 	/** Delegate handler for FEditorDelegates::PostUndoRedo. It internally calls DataLayerChanged.Broadcast and UpdateAllActorsVisibility to refresh the actors of each DataLayer. */
 	void PostUndoRedo();
 
+	void OnWorldPartitionInitialized(UWorldPartition* InWorldPartition);
+	void OnWorldPartitionUninitialized(UWorldPartition* InWorldPartition);
+	void OnLoadedActorAddedToLevel(AActor& InActor);
+	void OnDataLayerEditorLoadingStateChanged(bool bIsFromUserChange);
+
 	bool SetDataLayerIsLoadedInEditorInternal(UDataLayerInstance* DataLayer, const bool bIsLoadedInEditor, const bool bIsFromUserChange);
 
 	void BroadcastActorDataLayersChanged(const TWeakObjectPtr<AActor>& ChangedActor);
diff --git a/Engine/Source/Editor/SceneOutliner/Private/ActorBrowsingMode.cpp b/Engine/Source/Editor/SceneOutliner/Private/ActorBrowsingMode.cpp
--- a/Engine/Source/Editor/SceneOutliner/Private/ActorBrowsingMode.cpp
+++ b/Engine/Source/Editor/SceneOutliner/Private/ActorBrowsingMode.cpp
@@ -115,7 +115,7 @@ FActorBrowsingMode::FActorBrowsingMode(SSceneOutliner* InSceneOutliner, TWeakObj
 
 	FilterInfoMap.Add(TEXT("HideComponentsFilter"), HideComponentsInfo);
 
-	FSceneOutlinerFilterInfo HideLevelInstancesInfo(LOCTEXT("ToggleHideLevelInstances", "Hide Level Instances"), LOCTEXT("ToggleHideLevelInstancesToolTip", "When enabled, hides all level instance content."), SharedSettings->bHideLevelInstanceHierarchy, FCreateSceneOutlinerFilter::CreateStatic(&FActorBrowsingMode::CreateHideLevelInstancesFilter));
+	FSceneOutlinerFilterInfo HideLevelInstancesInfo(LOCTEXT("ToggleHideLevelInstanceContent", "Hide Level Instance Content"), LOCTEXT("ToggleHideLevelInstancesToolTip", "When enabled, hides all level instance content."), SharedSettings->bHideLevelInstanceHierarchy, FCreateSceneOutlinerFilter::CreateStatic(&FActorBrowsingMode::CreateHideLevelInstancesFilter));
 	HideLevelInstancesInfo.OnToggle().AddLambda([this](bool bIsActive)
 		{
 			UActorBrowsingModeSettings* Settings = GetMutableDefault<UActorBrowsingModeSettings>();
diff --git a/Engine/Source/Editor/SceneOutliner/Private/ActorDescTreeItem.cpp b/Engine/Source/Editor/SceneOutliner/Private/ActorDescTreeItem.cpp
--- a/Engine/Source/Editor/SceneOutliner/Private/ActorDescTreeItem.cpp
+++ b/Engine/Source/Editor/SceneOutliner/Private/ActorDescTreeItem.cpp
@@ -117,7 +117,7 @@ private:
 			if (const FWorldPartitionActorDesc* ActorDesc = TreeItem->ActorDescHandle.Get())
 			{
 				FFormatNamedArguments Args;
-				Args.Add(TEXT("ActorLabel"), FText::FromName(ActorDesc->GetActorLabel()));
+				Args.Add(TEXT("ActorLabel"), FText::FromString(TreeItem->GetDisplayString()));
 				Args.Add(TEXT("UnloadedTag"), LOCTEXT("UnloadedActorLabel", "(Unloaded)"));
 				return FText::Format(LOCTEXT("UnloadedActorDisplay", "{ActorLabel} {UnloadedTag}"), Args);
 			}
diff --git a/Engine/Source/Editor/SceneOutliner/Private/ActorHierarchy.cpp b/Engine/Source/Editor/SceneOutliner/Private/ActorHierarchy.cpp
--- a/Engine/Source/Editor/SceneOutliner/Private/ActorHierarchy.cpp
+++ b/Engine/Source/Editor/SceneOutliner/Private/ActorHierarchy.cpp
@@ -23,6 +23,7 @@
 #include "WorldPartition/WorldPartitionHelpers.h"
 #include "WorldPartition/WorldPartitionActorDesc.h"
 #include "WorldPartition/IWorldPartitionEditorModule.h"
+#include "WorldPartition/WorldPartitionSubsystem.h"
 #include "ActorFolder.h"
 
 TUniquePtr<FActorHierarchy> FActorHierarchy::Create(ISceneOutlinerMode* Mode, const TWeakObjectPtr<UWorld>& World)
@@ -49,6 +50,9 @@ TUniquePtr<FActorHierarchy> FActorHierarchy::Create(ISceneOutlinerMode* Mode, co
 			World->PersistentLevel->OnLoadedActorRemovedFromLevelEvent.AddRaw(Hierarchy, &FActorHierarchy::OnLoadedActorRemoved);
 		}
 
+		World->OnWorldPartitionInitialized().AddRaw(Hierarchy, &FActorHierarchy::OnWorldPartitionInitialized);
+		World->OnWorldPartitionUninitialized().AddRaw(Hierarchy, &FActorHierarchy::OnWorldPartitionUninitialized);
+
 		if (UWorldPartition* WorldPartition = World->GetWorldPartition())
 		{
 			WorldPartition->OnActorDescAddedEvent.AddRaw(Hierarchy, &FActorHierarchy::OnActorDescAdded);
@@ -98,6 +102,9 @@ FActorHierarchy::~FActorHierarchy()
 			RepresentingWorld->PersistentLevel->OnLoadedActorRemovedFromLevelEvent.RemoveAll(this);
 		}
 
+		RepresentingWorld->OnWorldPartitionInitialized().RemoveAll(this);
+		RepresentingWorld->OnWorldPartitionUninitialized().RemoveAll(this);
+
 		if (UWorldPartition* WorldPartition = RepresentingWorld->GetWorldPartition())
 		{
 			WorldPartition->OnActorDescAddedEvent.RemoveAll(this);
@@ -118,14 +125,15 @@ FActorHierarchy::~FActorHierarchy()
 	}
 }
 
-static FName GetFolderPathFromActorDesc(UWorld* InWorld, const FWorldPartitionActorDesc* InActorDesc)
+static FFolder GetFolderPathFromActorDesc(UWorld* InWorld, const FWorldPartitionActorDesc* InActorDesc)
 {
-	if (InWorld && InWorld->PersistentLevel->IsUsingActorFolders())
+	UWorld* World = InActorDesc->GetContainer() ? InActorDesc->GetContainer()->GetTypedOuter<UWorld>() : InWorld;
+	if (World && World->PersistentLevel->IsUsingActorFolders())
 	{
-		UActorFolder* ActorFolder = InWorld->PersistentLevel->GetActorFolder(InActorDesc->GetFolderGuid());
-		return ActorFolder ? ActorFolder->GetPath() : FFolder::GetEmptyPath();
+		UActorFolder* ActorFolder = World->PersistentLevel->GetActorFolder(InActorDesc->GetFolderGuid());
+		return ActorFolder ? ActorFolder->GetFolder() : FFolder();
 	}
-	return InActorDesc->GetFolderPath();
+	return FFolder(InActorDesc->GetFolderPath());
 }
 
 FSceneOutlinerTreeItemPtr FActorHierarchy::FindOrCreateParentItem(const ISceneOutlinerTreeItem& Item, const TMap<FSceneOutlinerTreeItemID, FSceneOutlinerTreeItemPtr>& Items, bool bCreate)
@@ -268,16 +276,16 @@ FSceneOutlinerTreeItemPtr FActorHierarchy::FindOrCreateParentItem(const ISceneOu
 	{
 		if (const FWorldPartitionActorDesc* ActorDesc = ActorDescItem->ActorDescHandle.Get())
 		{
-			const FName FolderPath = GetFolderPathFromActorDesc(RepresentingWorld.Get(), ActorDesc);
-			if (!FolderPath.IsNone())
+			const FFolder ActorDescFolder = GetFolderPathFromActorDesc(RepresentingWorld.Get(), ActorDesc);
+			if (!ActorDescFolder.IsNone())
 			{
-				if (const FSceneOutlinerTreeItemPtr* ParentItem = Items.Find(FFolder(FolderPath)))
+				if (const FSceneOutlinerTreeItemPtr* ParentItem = Items.Find(ActorDescFolder))
 				{
 					return *ParentItem;
 				}
 				else
 				{
-					return bCreate ? Mode->CreateItemFor<FActorFolderTreeItem>(FActorFolderTreeItem(FFolder(FolderPath), RepresentingWorld.Get()), true) : nullptr;
+					return bCreate ? Mode->CreateItemFor<FActorFolderTreeItem>(FActorFolderTreeItem(ActorDescFolder, RepresentingWorld), true) : nullptr;
 				}
 			}
 
@@ -285,9 +293,9 @@ FSceneOutlinerTreeItemPtr FActorHierarchy::FindOrCreateParentItem(const ISceneOu
 			const FGuid& ParentActorGuid = ActorDesc->GetSceneOutlinerParent();
 			if (ParentActorGuid.IsValid())
 			{
-				if (UWorldPartition* WorldPartition = RepresentingWorld->GetWorldPartition())
+				if (UActorDescContainer* ActorDescContainer = ActorDesc->GetContainer())
 				{
-					if (const FWorldPartitionActorDesc* ParentActorDesc = WorldPartition->GetActorDesc(ParentActorGuid))
+					if (const FWorldPartitionActorDesc* ParentActorDesc = ActorDescContainer->GetActorDesc(ParentActorGuid))
 					{
 						// If parent actor is loaded
 						if (AActor* ParentActor = ParentActorDesc->GetActor())
@@ -310,11 +318,28 @@ FSceneOutlinerTreeItemPtr FActorHierarchy::FindOrCreateParentItem(const ISceneOu
 						}
 						else
 						{
-							return bCreate ? Mode->CreateItemFor<FActorDescTreeItem>(FActorDescTreeItem(ParentActorGuid, WorldPartition)) : nullptr;
+							return bCreate ? Mode->CreateItemFor<FActorDescTreeItem>(FActorDescTreeItem(ParentActorGuid, ActorDescContainer)) : nullptr;
 						}
 					}
 				}
 			}
+			else if (UActorDescContainer* ActorDescContainer = ActorDesc->GetContainer())
+			{
+				const ULevelInstanceSubsystem* LevelInstanceSubsystem = UWorld::GetSubsystem<ULevelInstanceSubsystem>(RepresentingWorld.Get());
+				UWorld* OuterWorld = ActorDescContainer->GetTypedOuter<UWorld>();
+				// If parent actor is loaded
+				if (AActor* ParentActor = OuterWorld ? Cast<AActor>(LevelInstanceSubsystem->GetOwningLevelInstance(OuterWorld->PersistentLevel)) : nullptr)
+				{
+					if (const FSceneOutlinerTreeItemPtr* ParentItem = Items.Find(ParentActor))
+					{
+						return *ParentItem;
+					}
+					else
+					{
+						return bCreate ? Mode->CreateItemFor<FActorTreeItem>(ParentActor, true) : nullptr;
+					}
+				}
+			}
 		}
 	}
 
@@ -404,23 +429,37 @@ void FActorHierarchy::CreateWorldChildren(UWorld* World, TArray<FSceneOutlinerTr
 				// Create all component items
 				CreateComponentItems(Actor, OutItems);
 			}
-
 		}
 	}
 
 	if (bShowingUnloadedActors)
 	{
-		if (UWorldPartition* WorldPartition = World->GetWorldPartition())
+		if (UWorldPartitionSubsystem* WorldPartitionSubsystem = UWorld::GetSubsystem<UWorldPartitionSubsystem>(World))
 		{
-			FWorldPartitionHelpers::ForEachActorDesc(WorldPartition, [this, WorldPartition, &OutItems](const FWorldPartitionActorDesc* ActorDesc)
+			WorldPartitionSubsystem->ForEachWorldPartition([this, LevelInstanceSubsystem, &OutItems](UWorldPartition* WorldPartition)
 			{
-				if (ActorDesc != nullptr && !ActorDesc->IsLoaded(true))
+				// Skip unloaded actors if they are part of a non editing level instance and the outliner hides the content of level instances
+				if (!bShowingLevelInstances)
 				{
-					if (const FSceneOutlinerTreeItemPtr ActorDescItem = Mode->CreateItemFor<FActorDescTreeItem>(FActorDescTreeItem(ActorDesc->GetGuid(), WorldPartition)))
+					UWorld* OuterWorld = WorldPartition->GetTypedOuter<UWorld>();
+					ULevel* OuterLevel = OuterWorld ? OuterWorld->PersistentLevel : nullptr;
+					ILevelInstanceInterface* LevelInstance = LevelInstanceSubsystem ? LevelInstanceSubsystem->GetOwningLevelInstance(OuterLevel) : nullptr;
+					if (LevelInstance && !LevelInstance->IsEditing())
 					{
-						OutItems.Add(ActorDescItem);
+						return true;
 					}
 				}
+				FWorldPartitionHelpers::ForEachActorDesc(WorldPartition, [this, &OutItems](const FWorldPartitionActorDesc* ActorDesc)
+				{
+					if (ActorDesc != nullptr && !ActorDesc->IsLoaded(true))
+					{
+						if (const FSceneOutlinerTreeItemPtr ActorDescItem = Mode->CreateItemFor<FActorDescTreeItem>(FActorDescTreeItem(ActorDesc->GetGuid(), ActorDesc->GetContainer())))
+						{
+							OutItems.Add(ActorDescItem);
+						}
+					}
+					return true;
+				});
 				return true;
 			});
 		}
@@ -645,6 +684,16 @@ void FActorHierarchy::OnActorDescRemoved(FWorldPartitionActorDesc* ActorDesc)
 	}
 }
 
+void FActorHierarchy::OnWorldPartitionInitialized(UWorldPartition* InWorldPartition)
+{
+	FullRefreshEvent();
+}
+
+void FActorHierarchy::OnWorldPartitionUninitialized(UWorldPartition* InWorldPartition)
+{
+	FullRefreshEvent();
+}
+
 void FActorHierarchy::OnComponentsUpdated()
 {
 	FullRefreshEvent();
diff --git a/Engine/Source/Editor/SceneOutliner/Private/ActorTreeItem.cpp b/Engine/Source/Editor/SceneOutliner/Private/ActorTreeItem.cpp
--- a/Engine/Source/Editor/SceneOutliner/Private/ActorTreeItem.cpp
+++ b/Engine/Source/Editor/SceneOutliner/Private/ActorTreeItem.cpp
@@ -339,7 +339,8 @@ FActorTreeItem::FActorTreeItem(AActor* InActor)
 	, ID(InActor)
 {
 	check(InActor);
-	ActorLabel = InActor->GetActorLabel();
+
+	UpdateDisplayStringInternal();
 
 	Flags.bIsExpanded = InActor->bDefaultOutlinerExpansionState;
 	
@@ -359,7 +360,7 @@ FFolder::FRootObject FActorTreeItem::GetRootObject() const
 
 FString FActorTreeItem::GetDisplayString() const
 {
-	return ActorLabel;
+	return DisplayString;
 }
 
 bool FActorTreeItem::CanInteract() const
@@ -440,10 +441,7 @@ bool FActorTreeItem::GetPinnedState() const
 
 void FActorTreeItem::OnLabelChanged()
 {
-	if (Actor.IsValid())
-	{
-		ActorLabel = Actor->GetActorLabel();
-	}
+	UpdateDisplayString();
 }
 
 void FActorTreeItem::GenerateContextMenu(UToolMenu* Menu, SSceneOutliner& Outliner)
@@ -459,4 +457,14 @@ void FActorTreeItem::GenerateContextMenu(UToolMenu* Menu, SSceneOutliner& Outlin
 	}
 }
 
+void FActorTreeItem::UpdateDisplayString()
+{
+	UpdateDisplayStringInternal();
+}
+
+void FActorTreeItem::UpdateDisplayStringInternal()
+{
+	DisplayString = Actor.IsValid() ? Actor->GetActorLabel() : TEXT("None");
+}
+
 #undef LOCTEXT_NAMESPACE
diff --git a/Engine/Source/Editor/SceneOutliner/Private/SceneOutlinerModule.cpp b/Engine/Source/Editor/SceneOutliner/Private/SceneOutlinerModule.cpp
--- a/Engine/Source/Editor/SceneOutliner/Private/SceneOutlinerModule.cpp
+++ b/Engine/Source/Editor/SceneOutliner/Private/SceneOutlinerModule.cpp
@@ -24,6 +24,8 @@
 #include "WorldTreeItem.h"
 #include "WorldPartition/DataLayer/DataLayerInstance.h"
 #include "WorldPartition/DataLayer/DataLayerSubsystem.h"
+#include "LevelInstance/LevelInstanceSubsystem.h"
+#include "LevelInstance/LevelInstanceInterface.h"
 #include "Editor.h"
 
 #define LOCTEXT_NAMESPACE "SceneOutlinerModule"
@@ -236,41 +238,71 @@ void FSceneOutlinerModule::CreateActorInfoColumns(FSceneOutlinerInitializationOp
 	FGetTextForItem DataLayerInfoText = FGetTextForItem::CreateLambda([](const ISceneOutlinerTreeItem& Item) -> FString
 	{
 		TStringBuilder<128> Builder;
+		TSet<FString> DataLayerShortNames;
 
-		TArray<const UDataLayerInstance*> DataLayerInstances;
-
-		if (const FActorTreeItem* ActorItem = Item.CastTo<FActorTreeItem>())
+		auto BuildDataLayers = [&Builder, &DataLayerShortNames](const auto& DataLayerInstances, bool bPartOfOtherLevel)
 		{
-			AActor* Actor = ActorItem->Actor.Get();
-
-			if (Actor)
+			for (const UDataLayerInstance* DataLayerInstance : DataLayerInstances)
 			{
-				DataLayerInstances = Actor->GetDataLayerInstances();
+				bool bIsAlreadyInSet = false;
+				DataLayerShortNames.Add(DataLayerInstance->GetDataLayerShortName(), &bIsAlreadyInSet);
+				if (!bIsAlreadyInSet)
+				{
+					if (Builder.Len())
+					{
+						Builder += TEXT(", ");
+					}
+					// Put a '*' in front of DataLayers that are not part of of the main world
+					if (bPartOfOtherLevel)
+					{
+						Builder += "*";
+					}
+					Builder += DataLayerInstance->GetDataLayerShortName();
+				}
 			}
-		}
-		else if (const FActorDescTreeItem* ActorDescItem = Item.CastTo<FActorDescTreeItem>())
+		};
+
+		auto BuildDataLayersWithContext = [BuildDataLayers](const ISceneOutlinerTreeItem& Item, bool bUseLevelContext)
 		{
-			if (const FWorldPartitionActorDesc* ActorDesc = ActorDescItem->ActorDescHandle.Get(); ActorDesc && !ActorDesc->GetDataLayerInstanceNames().IsEmpty())
+			if (const FActorTreeItem* ActorItem = Item.CastTo<FActorTreeItem>())
 			{
-				const UActorDescContainer* ActorDescContainer = ActorDescItem->ActorDescHandle.Container.Get();
-				const UWorld* World = ActorDescContainer ? ActorDescContainer->GetWorld() : nullptr;
-				
-				if (const UDataLayerSubsystem* DataLayerSubsystem = UWorld::GetSubsystem<UDataLayerSubsystem>(World))
+				if (AActor* Actor = ActorItem->Actor.Get())
 				{
-					DataLayerInstances = DataLayerSubsystem->GetDataLayerInstances(ActorDesc->GetDataLayerInstanceNames());
+					BuildDataLayers(bUseLevelContext ? Actor->GetDataLayerInstancesForLevel() : Actor->GetDataLayerInstances(), bUseLevelContext);
 				}
 			}
-		}
-
-		for (const UDataLayerInstance* DataLayerInstance : DataLayerInstances)
-		{
-			if (Builder.Len())
+			else if (const FActorDescTreeItem* ActorDescItem = Item.CastTo<FActorDescTreeItem>())
 			{
-				Builder += TEXT(", ");
+				if (const FWorldPartitionActorDesc* ActorDesc = ActorDescItem->ActorDescHandle.Get(); ActorDesc && !ActorDesc->GetDataLayerInstanceNames().IsEmpty())
+				{
+					const UActorDescContainer* ActorDescContainer = ActorDescItem->ActorDescHandle.Container.Get();
+					const UWorld* World = ActorDescContainer ? ActorDescContainer->GetWorld() : nullptr;
+
+					if (const UDataLayerSubsystem* DataLayerSubsystem = UWorld::GetSubsystem<UDataLayerSubsystem>(World))
+					{
+						TSet<const UDataLayerInstance*> DataLayerInstances;
+						DataLayerInstances.Append(DataLayerSubsystem->GetDataLayerInstances(ActorDesc->GetDataLayerInstanceNames()));
+						if (ULevelInstanceSubsystem* LevelInstanceSubsystem = UWorld::GetSubsystem<ULevelInstanceSubsystem>(World))
+						{
+							UWorld* OuterWorld = ActorDescContainer->GetTypedOuter<UWorld>();
+							// Add parent container Data Layer Instances
+							AActor* CurrentActor = OuterWorld ? Cast<AActor>(LevelInstanceSubsystem->GetOwningLevelInstance(OuterWorld->PersistentLevel)) : nullptr;
+							while (CurrentActor)
+							{
+								DataLayerInstances.Append(bUseLevelContext ? CurrentActor->GetDataLayerInstancesForLevel() : CurrentActor->GetDataLayerInstances());
+								CurrentActor = Cast<AActor>(LevelInstanceSubsystem->GetParentLevelInstance(CurrentActor));
+							};
+						}
+						BuildDataLayers(DataLayerInstances, bUseLevelContext);
+					}
+				}
 			}
-			Builder += DataLayerInstance->GetDataLayerShortName();
-		}
-		
+		};
+
+		// List Actor's DataLayers part of the owning world, then those only part of the actor level
+		BuildDataLayersWithContext(Item, false);
+		BuildDataLayersWithContext(Item, true);
+
 		return Builder.ToString();
 	});
 
diff --git a/Engine/Source/Editor/SceneOutliner/Public/ActorDescTreeItem.h b/Engine/Source/Editor/SceneOutliner/Public/ActorDescTreeItem.h
--- a/Engine/Source/Editor/SceneOutliner/Public/ActorDescTreeItem.h
+++ b/Engine/Source/Editor/SceneOutliner/Public/ActorDescTreeItem.h
@@ -56,9 +56,11 @@ public:
 	/* End ISceneOutlinerTreeItem Implementation */
 	
 	const FGuid& GetGuid() const { return ActorGuid; }
-private:
-	void FocusActorBounds() const;
 
+protected:
 	FString DisplayString;
+
+private:
+	void FocusActorBounds() const;
 	FGuid ActorGuid;
 };
diff --git a/Engine/Source/Editor/SceneOutliner/Public/ActorHierarchy.h b/Engine/Source/Editor/SceneOutliner/Public/ActorHierarchy.h
--- a/Engine/Source/Editor/SceneOutliner/Public/ActorHierarchy.h
+++ b/Engine/Source/Editor/SceneOutliner/Public/ActorHierarchy.h
@@ -8,6 +8,7 @@
 
 class FWorldPartitionActorDesc;
 class UActorFolder;
+class UWorldPartition;
 
 class FActorHierarchy : public ISceneOutlinerHierarchy
 {
@@ -60,6 +61,9 @@ private:
 	void OnLevelAdded(ULevel* InLevel, UWorld* InWorld);
 	void OnLevelRemoved(ULevel* InLevel, UWorld* InWorld);
 
+	void OnWorldPartitionInitialized(UWorldPartition* InWorldPartition);
+	void OnWorldPartitionUninitialized(UWorldPartition* InWorldPartition);
+
 	/** Called when a folder is to be created */
 	void OnBroadcastFolderCreate(UWorld& InWorld, const FFolder& InNewFolder);
 
diff --git a/Engine/Source/Editor/SceneOutliner/Public/ActorTreeItem.h b/Engine/Source/Editor/SceneOutliner/Public/ActorTreeItem.h
--- a/Engine/Source/Editor/SceneOutliner/Public/ActorTreeItem.h
+++ b/Engine/Source/Editor/SceneOutliner/Public/ActorTreeItem.h
@@ -53,9 +53,16 @@ public:
 	virtual void OnLabelChanged() override;
 	virtual void GenerateContextMenu(UToolMenu* Menu, SSceneOutliner& Outliner) override;
 	/* End ISceneOutlinerTreeItem Implementation */
-public:
+
 	/** true if this item exists in both the current world and PIE. */
 	bool bExistsInCurrentWorldAndPIE;
-	/** Cached actor label */
-	FString ActorLabel;
+
+protected:
+	virtual void UpdateDisplayString();
+
+	/** Cached display string */
+	FString DisplayString;
+
+private:
+	void UpdateDisplayStringInternal();
 };
diff --git a/Engine/Source/Editor/UnrealEd/Private/EditorEngine.cpp b/Engine/Source/Editor/UnrealEd/Private/EditorEngine.cpp
--- a/Engine/Source/Editor/UnrealEd/Private/EditorEngine.cpp
+++ b/Engine/Source/Editor/UnrealEd/Private/EditorEngine.cpp
@@ -7212,7 +7212,8 @@ void UEditorEngine::OnAssetCreated(UObject* Asset)
 void UEditorEngine::InitializeNewlyCreatedInactiveWorld(UWorld* World)
 {
 	check(World);
-	if (!World->bIsWorldInitialized && World->WorldType == EWorldType::Inactive)
+
+	if (!World->bIsWorldInitialized && World->WorldType == EWorldType::Inactive && !World->IsInstanced())
 	{
 		const bool bOldDirtyState = World->GetOutermost()->IsDirty();
 
diff --git a/Engine/Source/Runtime/Engine/Classes/GameFramework/Actor.h b/Engine/Source/Runtime/Engine/Classes/GameFramework/Actor.h
--- a/Engine/Source/Runtime/Engine/Classes/GameFramework/Actor.h
+++ b/Engine/Source/Runtime/Engine/Classes/GameFramework/Actor.h
@@ -1184,7 +1184,7 @@ public:
 	bool RemoveAllDataLayers();
 	virtual bool SupportsDataLayer() const { return true; }
 	
-	bool HasValidDataLayers() const;
+	TArray<const UDataLayerInstance*> GetDataLayerInstancesForLevel() const;
 	TArray<FName> GetDataLayerInstanceNames() const;
 	bool IsPropertyChangedAffectingDataLayers(FPropertyChangedEvent& PropertyChangedEvent) const;
 	bool IsValidForDataLayer() const;
@@ -1221,6 +1221,9 @@ public:
 	UE_DEPRECATED(5.1, "Use GetDataLayerInstances() with no parameters instead")
 	TArray<const UDataLayerInstance*> GetDataLayerInstances(const AWorldDataLayers* WorldDataLayers) const { return TArray<const UDataLayerInstance*>(); }
 
+	UE_DEPRECATED(5.1, "Use HasDataLayers() instead")
+	bool HasValidDataLayers() const { return HasDataLayers(); }
+
 	//~ End Deprecated
 #endif
 
@@ -1231,6 +1234,7 @@ public:
 	bool HasDataLayers() const;
 
 private:
+	TArray<const UDataLayerInstance*> GetDataLayerInstancesInternal(bool bUseLevelContext) const;
 	bool UseWorldPartitionRuntimeCellDataLayers() const;
 
 	//~=============================================================================
diff --git a/Engine/Source/Runtime/Engine/Private/Actor.cpp b/Engine/Source/Runtime/Engine/Private/Actor.cpp
--- a/Engine/Source/Runtime/Engine/Private/Actor.cpp
+++ b/Engine/Source/Runtime/Engine/Private/Actor.cpp
@@ -57,6 +57,8 @@
 #if WITH_EDITOR
 #include "FoliageHelper.h"
 #include "AssetRegistry/AssetRegistryModule.h"
+#include "LevelInstance/LevelInstanceSubsystem.h"
+#include "LevelInstance/LevelInstanceInterface.h"
 #endif
 
 #include "WorldPartition/WorldPartition.h"
@@ -5908,6 +5910,15 @@ FArchive& operator<<(FArchive& Ar, AActor::FActorTransactionAnnotationData& Acto
 // DataLayers (begin)
 
 TArray<const UDataLayerInstance*> AActor::GetDataLayerInstances() const
+{
+	const bool bUseLevelContext = false;
+	return GetDataLayerInstancesInternal(bUseLevelContext);
+}
+
+// Returns all valid DataLayerInstances for this actor including those inherited from their parent level instance actor.
+// If bUseLevelContext is true, the actor level will be passed down to the DataLayerSubsystem which will use it 
+// to narrow down the resolving of valid datalayers for this particular level.
+TArray<const UDataLayerInstance*> AActor::GetDataLayerInstancesInternal(bool bUseLevelContext) const
 {
 	if (UseWorldPartitionRuntimeCellDataLayers())
 	{
@@ -5915,15 +5926,28 @@ TArray<const UDataLayerInstance*> AActor::GetDataLayerInstances() const
 	}
 
 #if WITH_EDITOR
-	if (UDataLayerSubsystem* DataLayerSubsystem = UWorld::GetSubsystem<UDataLayerSubsystem>(GetWorld()))
+	if (SupportsDataLayer() && IsValidForDataLayer())
 	{
 		TArray<const UDataLayerInstance*> DataLayerInstances;
-		DataLayerInstances += DataLayerSubsystem->GetDataLayerInstances(DataLayerAssets);
+		ULevel* OuterLevel = bUseLevelContext ? GetLevel() : nullptr;
+		if (UDataLayerSubsystem* DataLayerSubsystem = UWorld::GetSubsystem<UDataLayerSubsystem>(GetWorld()))
+		{
+			DataLayerInstances += DataLayerSubsystem->GetDataLayerInstances(DataLayerAssets, OuterLevel);
 
-		PRAGMA_DISABLE_DEPRECATION_WARNINGS
-			DataLayerInstances += DataLayerSubsystem->GetDataLayerInstances(DataLayers);
-		PRAGMA_ENABLE_DEPRECATION_WARNINGS
+			PRAGMA_DISABLE_DEPRECATION_WARNINGS
+			DataLayerInstances += DataLayerSubsystem->GetDataLayerInstances(DataLayers, OuterLevel);
+			PRAGMA_ENABLE_DEPRECATION_WARNINGS
+		}
 
+		// Add parent container Data Layer Instances
+		ULevelInstanceSubsystem* LevelInstanceSubsystem = UWorld::GetSubsystem<ULevelInstanceSubsystem>(GetWorld());
+		if (AActor* ParentActor = LevelInstanceSubsystem ? Cast<AActor>(LevelInstanceSubsystem->GetParentLevelInstance(this)) : nullptr)
+		{
+			for (const UDataLayerInstance* DataLayerInstance : ParentActor->GetDataLayerInstancesInternal(bUseLevelContext))
+			{
+				DataLayerInstances.AddUnique(DataLayerInstance);
+			}
+		}
 		return DataLayerInstances;
 	}
 #endif
@@ -5939,7 +5963,7 @@ bool AActor::ContainsDataLayer(const UDataLayerInstance* DataLayerInstance) cons
 	}
 
 #if WITH_EDITOR
-	return DataLayerInstance->ContainsActor(this);
+	return GetDataLayerInstances().Contains(DataLayerInstance);
 #else
 	return false;
 #endif
@@ -5947,16 +5971,29 @@ bool AActor::ContainsDataLayer(const UDataLayerInstance* DataLayerInstance) cons
 
 bool AActor::ContainsDataLayer(const UDataLayerAsset* DataLayerAsset) const
 {
+	if (!DataLayerAsset)
+	{
+		return false;
+	}
+
 	if (UseWorldPartitionRuntimeCellDataLayers())
 	{
 		return GetLevel()->GetWorldPartitionRuntimeCell()->ContainsDataLayer(DataLayerAsset);
 	}
 
 #if WITH_EDITOR
-	return DataLayerAsset && DataLayerAssets.Contains(DataLayerAsset);
-#else
-	return false;
+	if (DataLayerAssets.Contains(DataLayerAsset))
+	{
+		return true;
+	}
+	// Add parent container Data Layer Instances
+	ULevelInstanceSubsystem* LevelInstanceSubsystem = UWorld::GetSubsystem<ULevelInstanceSubsystem>(GetWorld());
+	if (AActor* ParentActor = LevelInstanceSubsystem ? Cast<AActor>(LevelInstanceSubsystem->GetParentLevelInstance(this)) : nullptr)
+	{
+		return ParentActor->ContainsDataLayer(DataLayerAsset);
+	}
 #endif
+	return false;
 }
 
 bool AActor::HasDataLayers() const
@@ -5967,7 +6004,7 @@ bool AActor::HasDataLayers() const
 	}
 
 #if WITH_EDITOR
-	return DataLayerAssets.Num() > 0 || DataLayers.Num() > 0;
+	return (GetDataLayerInstances().Num() > 0);
 #else
 	return false;
 #endif
diff --git a/Engine/Source/Runtime/Engine/Private/ActorEditor.cpp b/Engine/Source/Runtime/Engine/Private/ActorEditor.cpp
--- a/Engine/Source/Runtime/Engine/Private/ActorEditor.cpp
+++ b/Engine/Source/Runtime/Engine/Private/ActorEditor.cpp
@@ -16,6 +16,7 @@
 #include "WorldPartition/DataLayer/DataLayer.h"
 #include "WorldPartition/DataLayer/DataLayerSubsystem.h"
 #include "WorldPartition/DataLayer/DataLayerInstance.h"
+#include "WorldPartition/WorldPartitionRuntimeCellInterface.h"
 #include "EditorSupportDelegates.h"
 #include "Logging/TokenizedMessage.h"
 #include "Logging/MessageLog.h"
@@ -31,6 +32,7 @@
 #include "Engine/LevelStreaming.h"
 #include "WorldPartition/WorldPartitionActorDesc.h"
 #include "LevelInstance/LevelInstanceSubsystem.h"
+#include "LevelInstance/LevelInstanceInterface.h"
 #include "Folder.h"
 #include "ActorFolder.h"
 #include "WorldPersistentFolders.h"
@@ -1324,7 +1326,7 @@ bool AActor::AddDataLayer(const UDataLayerInstance* DataLayerInstance)
 bool AActor::AddDataLayer(const UDataLayerAsset* DataLayerAsset)
 {
 	bool bActorWasModified = false;
-	if (SupportsDataLayer() && DataLayerAsset && !ContainsDataLayer(DataLayerAsset))
+	if (SupportsDataLayer() && DataLayerAsset && !DataLayerAssets.Contains(DataLayerAsset))
 	{
 		Modify();
 		bActorWasModified = true;
@@ -1341,7 +1343,7 @@ bool AActor::RemoveDataLayer(const UDataLayerInstance* DataLayerInstance)
 bool AActor::RemoveDataLayer(const UDataLayerAsset* DataLayerAsset)
 {
 	bool bActorWasModified = false;
-	if (ContainsDataLayer(DataLayerAsset))
+	if (DataLayerAsset && DataLayerAssets.Contains(DataLayerAsset))
 	{
 		if (!bActorWasModified)
 		{
@@ -1356,7 +1358,7 @@ bool AActor::RemoveDataLayer(const UDataLayerAsset* DataLayerAsset)
 
 bool AActor::RemoveAllDataLayers()
 {
-	if (HasDataLayers())
+	if ((DataLayerAssets.Num() > 0) || (DataLayers.Num() > 0))
 	{
 		Modify();
 		DataLayerAssets.Empty();
@@ -1366,44 +1368,24 @@ bool AActor::RemoveAllDataLayers()
 	return false;
 }
 
-bool AActor::HasValidDataLayers() const
-{
-	if (UDataLayerSubsystem* DataLayerSubsystem = UWorld::GetSubsystem<UDataLayerSubsystem>(GetWorld()))
-	{
-		for (const TObjectPtr<const UDataLayerAsset>& DataLayerAsset : DataLayerAssets)
-		{
-			if (const UDataLayerInstance* DataLayerInstance = DataLayerSubsystem->GetDataLayerInstance(DataLayerAsset))
-			{
-				return true;
-			}
-		}
-
-		for (const FActorDataLayer& ActorDataLayer : DataLayers)
-		{
-			if (const UDataLayerInstance* DataLayerInstance = DataLayerSubsystem->GetDataLayerInstance(ActorDataLayer.Name))
-			{
-				return true;
-			}
-		}
-	}
-	return false;
-}
-
 TArray<FName> AActor::GetDataLayerInstanceNames() const
 {
 	TArray<FName> DataLayerInstanceNames;
-	
-	if (UDataLayerSubsystem* DataLayerSubsystem = UWorld::GetSubsystem<UDataLayerSubsystem>(GetWorld()))
+	TArray<const UDataLayerInstance*> DataLayerInstances = GetDataLayerInstances();
+	DataLayerInstanceNames.Reserve(DataLayerInstances.Num());
+	for (const UDataLayerInstance* DataLayerInstance : DataLayerInstances)
 	{
-		DataLayerInstanceNames += DataLayerSubsystem->GetDataLayerInstanceNames(DataLayerAssets);
-
-		PRAGMA_DISABLE_DEPRECATION_WARNINGS
-		DataLayerInstanceNames += DataLayerSubsystem->GetDataLayerInstanceNames(DataLayers);
-		PRAGMA_ENABLE_DEPRECATION_WARNINGS
+		DataLayerInstanceNames.Add(DataLayerInstance->GetDataLayerFName());
 	}
 	return DataLayerInstanceNames;
 }
 
+TArray<const UDataLayerInstance*> AActor::GetDataLayerInstancesForLevel() const
+{
+	const bool bUseLevelContext = true;
+	return GetDataLayerInstancesInternal(bUseLevelContext);
+}
+
 void AActor::FixupDataLayers(bool bRevertChangesOnLockedDataLayer /*= false*/)
 {
 	if (!GetPackage()->HasAnyPackageFlags(PKG_PlayInEditor))
@@ -1422,18 +1404,21 @@ void AActor::FixupDataLayers(bool bRevertChangesOnLockedDataLayer /*= false*/)
 				return;
 			}
 
+			ULevel* Level = GetLevel();
+
 			if (bRevertChangesOnLockedDataLayer)
 			{
 				// Since it's not possible to prevent changes of particular elements of an array, rollback change on locked DataLayers.
 				TSet<const UDataLayerAsset*> PreEdit(PreEditChangeDataLayers);
 				TSet<const UDataLayerAsset*> PostEdit(DataLayerAssets);
 
-				auto DifferenceContainsLockedDataLayers = [DataLayerSubsystem](const TSet<const UDataLayerAsset*>& A, const TSet<const UDataLayerAsset*>& B)
+				auto DifferenceContainsLockedDataLayers = [DataLayerSubsystem, Level](const TSet<const UDataLayerAsset*>& A, const TSet<const UDataLayerAsset*>& B)
 				{
 					TSet<const UDataLayerAsset*> Diff = A.Difference(B);
 					for (const UDataLayerAsset* DataLayerAsset : Diff)
 					{
-						const UDataLayerInstance* DataLayerInstance = DataLayerSubsystem->GetDataLayerInstance(DataLayerAsset);
+						// We pass Actor Level when resolving the DataLayerInstance as we do the fixup relative to this level
+						const UDataLayerInstance* DataLayerInstance = DataLayerSubsystem->GetDataLayerInstance(DataLayerAsset, Level);
 						if (DataLayerInstance && DataLayerInstance->IsLocked())
 						{
 							return true;
@@ -1450,15 +1435,16 @@ void AActor::FixupDataLayers(bool bRevertChangesOnLockedDataLayer /*= false*/)
 			}
 
 			PRAGMA_DISABLE_DEPRECATION_WARNINGS
-			auto CleanupDataLayers = [this, DataLayerSubsystem](auto& DataLayerArray)
+			auto CleanupDataLayers = [this, DataLayerSubsystem, Level](auto& DataLayerArray)
 			{
 				using ArrayType = typename TRemoveReference<decltype(DataLayerArray)>::Type;
 					
 				ArrayType ExistingDataLayer;
 				for (int32 Index = 0; Index < DataLayerArray.Num();)
 				{
+					// We pass Actor Level when resolving the DataLayerInstance as we do the fixup relative to this level
 					auto& DataLayer = DataLayerArray[Index];
-					if (!DataLayerSubsystem->GetDataLayerInstance(DataLayer) || ExistingDataLayer.Contains(DataLayer))
+					if (!DataLayerSubsystem->GetDataLayerInstance(DataLayer, Level) || ExistingDataLayer.Contains(DataLayer))
 					{
 						DataLayerArray.RemoveAtSwap(Index);
 					}
diff --git a/Engine/Source/Runtime/Engine/Private/LevelInstance/LevelInstanceEditorLevelStreaming.cpp b/Engine/Source/Runtime/Engine/Private/LevelInstance/LevelInstanceEditorLevelStreaming.cpp
--- a/Engine/Source/Runtime/Engine/Private/LevelInstance/LevelInstanceEditorLevelStreaming.cpp
+++ b/Engine/Source/Runtime/Engine/Private/LevelInstance/LevelInstanceEditorLevelStreaming.cpp
@@ -96,6 +96,11 @@ ULevelStreamingLevelInstanceEditor* ULevelStreamingLevelInstanceEditor::Load(ILe
 
 		GEngine->BlockTillLevelStreamingCompleted(LevelInstanceActor->GetWorld());
 
+		if (ULevel* LoadedLevel = LevelStreaming->GetLoadedLevel())
+		{
+			LoadedLevel->OnLoadedActorAddedToLevelEvent.AddUObject(LevelStreaming, &ULevelStreamingLevelInstanceEditor::OnLoadedActorAddedToLevel);
+		}
+
 		// Create special actor that will handle changing the pivot of this level
 		ALevelInstancePivot::Create(LevelInstance, LevelStreaming);
 
@@ -109,10 +114,19 @@ void ULevelStreamingLevelInstanceEditor::Unload(ULevelStreamingLevelInstanceEdit
 {
 	if (ULevelInstanceSubsystem* LevelInstanceSubsystem = LevelStreaming->GetWorld()->GetSubsystem<ULevelInstanceSubsystem>())
 	{
-		LevelInstanceSubsystem->RemoveLevelsFromWorld({ LevelStreaming->GetLoadedLevel() });
+		if (ULevel* LoadedLevel = LevelStreaming->GetLoadedLevel())
+		{
+			LoadedLevel->OnLoadedActorAddedToLevelEvent.RemoveAll(LevelStreaming);
+			LevelInstanceSubsystem->RemoveLevelsFromWorld({ LoadedLevel });
+		}
 	}
 }
 
+void ULevelStreamingLevelInstanceEditor::OnLoadedActorAddedToLevel(AActor& InActor)
+{
+	OnLevelActorAdded(&InActor);
+}
+
 void ULevelStreamingLevelInstanceEditor::OnLevelActorAdded(AActor* InActor)
 {
 	if (InActor && InActor->GetLevel() == LoadedLevel)
diff --git a/Engine/Source/Runtime/Engine/Private/LevelInstance/LevelInstanceLevelStreaming.cpp b/Engine/Source/Runtime/Engine/Private/LevelInstance/LevelInstanceLevelStreaming.cpp
--- a/Engine/Source/Runtime/Engine/Private/LevelInstance/LevelInstanceLevelStreaming.cpp
+++ b/Engine/Source/Runtime/Engine/Private/LevelInstance/LevelInstanceLevelStreaming.cpp
@@ -54,6 +54,59 @@ FBox ULevelStreamingLevelInstance::GetBounds() const
 	check(GetLoadedLevel());
 	return ALevelBounds::CalculateLevelBounds(GetLoadedLevel());
 }
+
+void ULevelStreamingLevelInstance::OnLoadedActorAddedToLevel(AActor& InActor)
+{
+	check(LevelInstanceEditorInstanceActor.IsValid());
+	PrepareLevelInstanceLoadedActor(InActor, GetLevelInstance());
+}
+
+void ULevelStreamingLevelInstance::PrepareLevelInstanceLoadedActor(AActor& InActor, ILevelInstanceInterface* InLevelInstance)
+{
+	InActor.ClearFlags(RF_Transactional);
+	InActor.SetFlags(RF_Transient);
+
+	if (InActor.IsPackageExternal())
+	{
+		ResetLoaders(InActor.GetExternalPackage());
+		InActor.GetPackage()->SetFlags(RF_Transient);
+	}
+
+	InActor.PushSelectionToProxies();
+	if (InLevelInstance)
+	{
+		InActor.PushLevelInstanceEditingStateToProxies(CastChecked<AActor>(InLevelInstance)->IsInEditingLevelInstance());
+	}
+
+	if (LevelInstanceEditorInstanceActor.IsValid())
+	{
+		if (InActor.GetAttachParentActor() == nullptr && !InActor.IsChildActor())
+		{
+			// Detect if LevelInstance has moved since it was loaded/created and move the child actor accordingly
+			if (!LevelTransform.Equals(LevelInstanceEditorInstanceActor->GetTransform()))
+			{
+				FTransform TransformToApply = LevelTransform.Inverse() * LevelInstanceEditorInstanceActor->GetTransform();
+				FLevelUtils::FApplyLevelTransformParams TransformParams(InActor.GetLevel(), TransformToApply);
+				TransformParams.Actor = &InActor;
+				TransformParams.bDoPostEditMove = true;
+				FLevelUtils::ApplyLevelTransform(TransformParams);
+			}
+
+			InActor.AttachToActor(LevelInstanceEditorInstanceActor.Get(), FAttachmentTransformRules::KeepWorldTransform);
+		}
+	}
+}
+
+void ULevelStreamingLevelInstance::OnLoadedActorRemovedFromLevel(AActor& InActor)
+{
+	// Detach actor or else it will keep it alive (attachement keeps a reference to the actor)
+	check(LevelInstanceEditorInstanceActor.IsValid());
+	if (InActor.GetAttachParentActor() == LevelInstanceEditorInstanceActor.Get())
+	{
+		InActor.DetachFromActor(FDetachmentTransformRules::KeepWorldTransform);
+	}
+}
+
 #endif
 
 ULevelStreamingLevelInstance* ULevelStreamingLevelInstance::LoadInstance(ILevelInstanceInterface* LevelInstance)
@@ -99,10 +152,21 @@ ULevelStreamingLevelInstance* ULevelStreamingLevelInstance::LoadInstance(ILevelI
 			check(Level);
 			check(LevelStreaming->GetCurrentState() == ULevelStreaming::ECurrentState::LoadedVisible);
 
+			Level->OnLoadedActorAddedToLevelEvent.AddUObject(LevelStreaming, &ULevelStreamingLevelInstance::OnLoadedActorAddedToLevel);
+			Level->OnLoadedActorRemovedFromLevelEvent.AddUObject(LevelStreaming, &ULevelStreamingLevelInstance::OnLoadedActorRemovedFromLevel);
+
 			UWorld* OuterWorld = Level->GetTypedOuter<UWorld>();
 			OuterWorld->ClearFlags(RF_Transactional);
 			OuterWorld->SetFlags(RF_Transient);
-			ResetLoaders(OuterWorld->GetPackage());
+			
+			// @todo_ow: Resetting the load will prevent any OFPA package to properly load since their import level will fail to resolve.
+			//           This is a temporary workaround. The downside of this is that the level can't be saved. 
+			//           Most of the changes will only affect OFPA packages. One problematic use case is when changing the pivot of a Level Instance.
+			FName PackageName = OuterWorld->GetPackage()->GetLoadedPath().GetPackageFName();
+			if (!ULevel::GetIsLevelPartitionedFromPackage(PackageName))
+			{
+				ResetLoaders(OuterWorld->GetPackage());
+			}
 
 			OuterWorld->GetPackage()->ClearFlags(RF_Transactional);
 			OuterWorld->GetPackage()->SetFlags(RF_Transient);
@@ -117,14 +181,7 @@ ULevelStreamingLevelInstance* ULevelStreamingLevelInstance::LoadInstance(ILevelI
 			{
 				if (LevelActor)
 				{
-					if (LevelActor->IsPackageExternal())
-					{
-						ResetLoaders(LevelActor->GetExternalPackage());
-						LevelActor->GetPackage()->SetFlags(RF_Transient);
-					}
-
-					LevelActor->PushSelectionToProxies();
-					LevelActor->PushLevelInstanceEditingStateToProxies(CastChecked<AActor>(LevelInstance)->IsInEditingLevelInstance());
+					LevelStreaming->PrepareLevelInstanceLoadedActor(*LevelActor, LevelInstance);
 				}
 			}
 			Level->ForEachActorFolder([](UActorFolder* ActorFolder)
@@ -138,7 +195,7 @@ ULevelStreamingLevelInstance* ULevelStreamingLevelInstance::LoadInstance(ILevelI
 			});
 
 			// Create special actor that will handle selection and transform
-			ALevelInstanceEditorInstanceActor::Create(LevelInstance, Level);
+			LevelStreaming->LevelInstanceEditorInstanceActor = ALevelInstanceEditorInstanceActor::Create(LevelInstance, Level);
 		}
 #endif
 		return LevelStreaming;
@@ -158,8 +215,12 @@ void ULevelStreamingLevelInstance::UnloadInstance(ULevelStreamingLevelInstance*
 #if WITH_EDITOR
 	else
 	{
-		// Check if we need to flush the Trans buffer...
 		ULevel* LoadedLevel = LevelStreaming->GetLoadedLevel();
+		LoadedLevel->OnLoadedActorAddedToLevelEvent.RemoveAll(LevelStreaming);
+		LoadedLevel->OnLoadedActorRemovedFromLevelEvent.RemoveAll(LevelStreaming);
+		LevelStreaming->LevelInstanceEditorInstanceActor.Reset();
+
+		// Check if we need to flush the Trans buffer...
 		UWorld* OuterWorld = LoadedLevel->GetTypedOuter<UWorld>();
 		bool bResetTrans = false;
 		ForEachObjectWithOuterBreakable(OuterWorld, [&bResetTrans](UObject* Obj)
diff --git a/Engine/Source/Runtime/Engine/Private/LevelInstance/LevelInstanceSubsystem.cpp b/Engine/Source/Runtime/Engine/Private/LevelInstance/LevelInstanceSubsystem.cpp
--- a/Engine/Source/Runtime/Engine/Private/LevelInstance/LevelInstanceSubsystem.cpp
+++ b/Engine/Source/Runtime/Engine/Private/LevelInstance/LevelInstanceSubsystem.cpp
@@ -21,6 +21,9 @@
 #include "LevelInstance/LevelInstanceEditorInstanceActor.h"
 #include "LevelInstance/LevelInstanceEditorObject.h"
 #include "LevelInstance/LevelInstanceEditorPivotActor.h"
+#include "WorldPartition/DataLayer/DataLayerSubsystem.h"
+#include "WorldPartition/DataLayer/DataLayerInstanceWithAsset.h"
+#include "WorldPartition/WorldPartitionMiniMap.h"
 #include "Misc/ScopedSlowTask.h"
 #include "Misc/ITransaction.h"
 #include "Misc/Paths.h"
@@ -869,20 +872,50 @@ ILevelInstanceInterface* ULevelInstanceSubsystem::CreateLevelInstanceFrom(const
 	{
 		const bool bIsPartitioned = GetWorld()->IsPartitionedWorld();
 		LevelStreaming = StaticCast<ULevelStreamingLevelInstanceEditor*>(EditorLevelUtils::CreateNewStreamingLevelForWorld(
-		*GetWorld(), ULevelStreamingLevelInstanceEditor::StaticClass(), CreationParams.UseExternalActors(), LevelFilename, &ActorsToMove, CreationParams.TemplateWorld, /*bUseSaveAs*/true, bIsPartitioned, [bIsPartitioned](ULevel* InLevel)
+		*GetWorld(), ULevelStreamingLevelInstanceEditor::StaticClass(), CreationParams.UseExternalActors(), LevelFilename, &ActorsToMove, CreationParams.TemplateWorld, /*bUseSaveAs*/true, bIsPartitioned, [this, bIsPartitioned, &ActorsToMove](ULevel* InLevel)
 		{
+			if (InLevel->IsUsingExternalActors())
+			{
+				// UWorldFactory::FactoryCreateNew will modify the default brush to be in global space (see GEditor->InitBuilderBrush(NewWorld)).
+				// The level is about to be saved, it doesn't have a transform and it is not yet added to the world levels.
+				// Since, no logic will remove the transform on the actor, force its transform to identity here.
+				if (ABrush* Brush = InLevel->GetDefaultBrush())
+				{
+					Brush->GetRootComponent()->SetRelativeTransform(FTransform::Identity);
+				}
+			}
 			if (bIsPartitioned)
 			{
 				UWorldPartition* WorldPartition = InLevel->GetWorldPartition();
-				if (ensure(WorldPartition))
+				check(WorldPartition);
+				
+				// Flag the world partition that it can be used by a Level Instance
+				WorldPartition->SetCanBeUsedByLevelInstance(true);
+
+				// Make sure new level's AWorldDataLayers contains all the necessary Data Layer Instances before moving actors
+				if (UDataLayerSubsystem* DataLayerSubsystem = UWorld::GetSubsystem<UDataLayerSubsystem>(GetWorld()))
 				{
-					// Flag the world partition that it can be used by a Level Instance
-					WorldPartition->SetCanBeUsedByLevelInstance(true);
+					TSet<TObjectPtr<const UDataLayerAsset>> SourceDataLayerAssets;
+					for (AActor* ActorToMove : ActorsToMove)
+					{
+						// Use the raw asset list as we don't want parent DataLayers
+						SourceDataLayerAssets.Append(ActorToMove->GetDataLayerAssets());
+					}
+					AWorldDataLayers* WorldDataLayers = InLevel->GetWorldDataLayers();
+					check(WorldDataLayers);
 
-					// Validation
-					check(InLevel->IsUsingActorFolders());
-					check(!WorldPartition->IsStreamingEnabled());
+					for (const UDataLayerAsset* SourceDataLayerAsset : SourceDataLayerAssets)
+					{
+						if (UDataLayerInstanceWithAsset* SourceDataLayerInstance = Cast<UDataLayerInstanceWithAsset>(DataLayerSubsystem->GetDataLayerInstanceFromAsset(SourceDataLayerAsset)))
+						{
+							UDataLayerInstanceWithAsset* DataLayerInstanceWithAsset = WorldDataLayers->CreateDataLayer<UDataLayerInstanceWithAsset>(SourceDataLayerAsset);
+						}
+					}
 				}
+
+				// Validation
+				check(InLevel->IsUsingActorFolders());
+				check(!WorldPartition->IsStreamingEnabled());
 			}
 		}));
 	}
@@ -1088,7 +1121,13 @@ void ULevelInstanceSubsystem::BreakLevelInstance_Impl(ILevelInstanceInterface* L
 			}
 
 			// Skip some actor types
-			if (!Actor->IsA<ALevelBounds>() && (Actor != Actor->GetLevel()->GetDefaultBrush()) && !Actor->IsA<AWorldSettings>() && !Actor->IsA<ALevelInstanceEditorInstanceActor>())
+			// @todo_ow: Move this logic in a new virtual function
+			if ((Actor != Actor->GetLevel()->GetDefaultBrush()) &&
+				!Actor->IsA<ALevelBounds>() &&
+				!Actor->IsA<AWorldSettings>() &&
+				!Actor->IsA<ALevelInstanceEditorInstanceActor>() &&
+				!Actor->IsA<AWorldDataLayers>() &&
+				!Actor->IsA<AWorldPartitionMiniMap>())
 			{
 				if (CanMoveActorToLevel(Actor))
 				{
@@ -1467,7 +1506,7 @@ FLevelInstanceID ULevelInstanceSubsystem::FLevelInstanceEdit::GetLevelInstanceID
 
 const ULevelInstanceSubsystem::FLevelInstanceEdit* ULevelInstanceSubsystem::GetLevelInstanceEdit(const ILevelInstanceInterface* LevelInstance) const
 {
-	if (LevelInstanceEdit && LevelInstanceEdit->GetLevelInstanceID() == LevelInstance->GetLevelInstanceID())
+	if (LevelInstance && LevelInstanceEdit && LevelInstanceEdit->GetLevelInstanceID() == LevelInstance->GetLevelInstanceID())
 	{
 		return LevelInstanceEdit.Get();
 	}
@@ -1964,6 +2003,48 @@ TArray<ILevelInstanceInterface*> ULevelInstanceSubsystem::GetLevelInstances(cons
 	return MatchingLevelInstances;
 }
 
+void ULevelInstanceSubsystem::ForEachLevelInstanceActorAncestors(const ULevel* Level, TFunctionRef<bool(AActor*)> Operation) const
+{
+	AActor* CurrentActor = Cast<AActor>(GetOwningLevelInstance(Level));
+	while (CurrentActor)
+	{
+		if (!Operation(CurrentActor))
+		{
+			break;
+		}
+		CurrentActor = Cast<AActor>(GetParentLevelInstance(CurrentActor));
+	};
+}
+
+TArray<AActor*> ULevelInstanceSubsystem::GetParentLevelInstanceActors(const ULevel* Level) const
+{
+	TArray<AActor*> ParentActors;
+	ForEachLevelInstanceActorAncestors(Level, [&ParentActors](AActor* Parent)
+	{
+		ParentActors.Add(Parent);
+		return true;
+	});
+	return ParentActors;
+}
+
+// Builds and returns a string based the LevelInstance hierarchy using actor labels.
+// Ex: ParentLevelInstanceActorLabel.ChildLevelInstanceActorLabel.ActorLabel
+FString ULevelInstanceSubsystem::PrefixWithParentLevelInstanceActorLabels(const FString& ActorLabel, const ULevel* Level) const
+{
+	TStringBuilder<128> Builder;
+	Builder = ActorLabel;
+	ForEachLevelInstanceActorAncestors(Level, [&Builder](AActor* Parent)
+	{
+		if (Builder.Len())
+		{
+			Builder.Prepend(TEXT("."));
+		}
+		Builder.Prepend(Parent->GetActorLabel());
+		return true;
+	});
+	return Builder.ToString();
+}
+
 bool ULevelInstanceSubsystem::CheckForLoop(const ILevelInstanceInterface* LevelInstance, TArray<TPair<FText, TSoftObjectPtr<UWorld>>>* LoopInfo, const ILevelInstanceInterface** LoopStart)
 {
 	return CheckForLoop(LevelInstance, LevelInstance->GetWorldAsset(), LoopInfo, LoopStart);
diff --git a/Engine/Source/Runtime/Engine/Private/WorldPartition/ActorDescContainer.cpp b/Engine/Source/Runtime/Engine/Private/WorldPartition/ActorDescContainer.cpp
--- a/Engine/Source/Runtime/Engine/Private/WorldPartition/ActorDescContainer.cpp
+++ b/Engine/Source/Runtime/Engine/Private/WorldPartition/ActorDescContainer.cpp
@@ -275,9 +275,17 @@ void UActorDescContainer::LoadAllActors(TArray<FWorldPartitionReference>& OutRef
 	}
 }
 
+bool UActorDescContainer::ShouldRegisterDelegates()
+{
+	UWorld* OuterWorld = GetTypedOuter<UWorld>();
+	// No need to register delegates for level instances
+	bool bIsInstance = OuterWorld && OuterWorld->IsInstanced() && !OuterWorld->GetPackage()->HasAnyPackageFlags(PKG_NewlyCreated);
+	return GEditor && !IsTemplate() && World && !World->IsGameWorld() && !bIsInstance;
+}
+
 void UActorDescContainer::RegisterEditorDelegates()
 {
-	if (GEditor && !IsTemplate() && World && !World->IsGameWorld())
+	if (ShouldRegisterDelegates())
 	{
 		FWorldDelegates::OnPostWorldRename.AddUObject(this, &UActorDescContainer::OnWorldRenamed);
 		FCoreUObjectDelegates::OnObjectPreSave.AddUObject(this, &UActorDescContainer::OnObjectPreSave);
@@ -288,7 +296,7 @@ void UActorDescContainer::RegisterEditorDelegates()
 
 void UActorDescContainer::UnregisterEditorDelegates()
 {
-	if (GEditor && !IsTemplate() && World && !World->IsGameWorld())
+	if (ShouldRegisterDelegates())
 	{
 		FWorldDelegates::OnPostWorldRename.RemoveAll(this);
 		FCoreUObjectDelegates::OnObjectPreSave.RemoveAll(this);
diff --git a/Engine/Source/Runtime/Engine/Private/WorldPartition/DataLayer/DataLayerInstance.cpp b/Engine/Source/Runtime/Engine/Private/WorldPartition/DataLayer/DataLayerInstance.cpp
--- a/Engine/Source/Runtime/Engine/Private/WorldPartition/DataLayer/DataLayerInstance.cpp
+++ b/Engine/Source/Runtime/Engine/Private/WorldPartition/DataLayer/DataLayerInstance.cpp
@@ -84,6 +84,21 @@ bool UDataLayerInstance::IsEffectiveLoadedInEditor() const
 	return bResult;
 }
 
+bool UDataLayerInstance::CanEditChange(const FProperty* InProperty) const
+{
+	if (!Super::CanEditChange(InProperty))
+	{
+		return false;
+	}
+
+	if (IsReadOnly())
+	{
+		return false;
+	}
+
+	return true;
+}
+
 bool UDataLayerInstance::IsLocked() const
 {
 	if (bIsLocked)
@@ -103,7 +118,7 @@ bool UDataLayerInstance::CanParent(const UDataLayerInstance* InParent) const
 {
 	return (this != InParent)
 		&& (Parent != InParent) 
-		&& (InParent == nullptr || IsDataLayerTypeValidToParent(InParent->GetType()));
+		&& (InParent == nullptr || (IsDataLayerTypeValidToParent(InParent->GetType()) && (InParent->GetOuterAWorldDataLayers() == GetOuterAWorldDataLayers())));
 }
 
 bool UDataLayerInstance::IsDataLayerTypeValidToParent(EDataLayerType ParentDataLayerType) const
diff --git a/Engine/Source/Runtime/Engine/Private/WorldPartition/DataLayer/DataLayerInstanceWithAsset.cpp b/Engine/Source/Runtime/Engine/Private/WorldPartition/DataLayer/DataLayerInstanceWithAsset.cpp
--- a/Engine/Source/Runtime/Engine/Private/WorldPartition/DataLayer/DataLayerInstanceWithAsset.cpp
+++ b/Engine/Source/Runtime/Engine/Private/WorldPartition/DataLayer/DataLayerInstanceWithAsset.cpp
@@ -30,19 +30,44 @@ void UDataLayerInstanceWithAsset::OnCreated(const UDataLayerAsset* Asset)
 	SetVisible(true);
 }
 
-bool UDataLayerInstanceWithAsset::AddActor(AActor* Actor) const
+bool UDataLayerInstanceWithAsset::IsReadOnly() const
 {
-	return Actor->AddDataLayer(DataLayerAsset);
+	return !GetOuterAWorldDataLayers()->IsMainWorldDataLayers();
 }
 
-bool UDataLayerInstanceWithAsset::RemoveActor(AActor* Actor) const
+bool UDataLayerInstanceWithAsset::IsLocked() const
 {
-	return Actor->RemoveDataLayer(DataLayerAsset);
+	if (Super::IsLocked())
+	{
+		return true;
+	}
+	return IsReadOnly();
+}
+
+bool UDataLayerInstanceWithAsset::AddActor(AActor* Actor) const
+{
+	// If actor's level WorldDataLayers doesn't match this DataLayerInstance outer WorldDataLayers, 
+	// Make sure that an instance for this Data Layer asset exists in the outer level AWorldDataLayer and forward the call on it.
+	AWorldDataLayers* OuterWorldDataLayers = Actor->GetLevel()->GetWorldDataLayers();
+	if (GetOuterAWorldDataLayers() != OuterWorldDataLayers)
+	{
+		check(OuterWorldDataLayers);
+		UDataLayerInstanceWithAsset* OuterDataLayerInstanceWithAsset = Cast<UDataLayerInstanceWithAsset>(const_cast<UDataLayerInstance*>(OuterWorldDataLayers->GetDataLayerInstance(DataLayerAsset)));
+		if (!OuterDataLayerInstanceWithAsset)
+		{
+			OuterDataLayerInstanceWithAsset = OuterWorldDataLayers->CreateDataLayer<UDataLayerInstanceWithAsset>(DataLayerAsset);
+		}
+		return OuterDataLayerInstanceWithAsset->AddActor(Actor);
+	}
+	else
+	{
+		return Actor->AddDataLayer(DataLayerAsset);
+	}
 }
 
-bool UDataLayerInstanceWithAsset::ContainsActor(const AActor* Actor) const
+bool UDataLayerInstanceWithAsset::RemoveActor(AActor* Actor) const
 {
-	return Actor->ContainsDataLayer(DataLayerAsset);
+	return Actor->RemoveDataLayer(DataLayerAsset);
 }
 
 bool UDataLayerInstanceWithAsset::Validate(IStreamingGenerationErrorHandler* ErrorHandler) const
@@ -72,7 +97,7 @@ bool UDataLayerInstanceWithAsset::Validate(IStreamingGenerationErrorHandler* Err
 		}
 
 		return true;
-	});
+	}, GetOuterAWorldDataLayers()->GetLevel()); // Resolve DataLayerInstances based on outer level
 
 	bIsValid &= Super::Validate(ErrorHandler);
 
diff --git a/Engine/Source/Runtime/Engine/Private/WorldPartition/DataLayer/DataLayerSubsystem.cpp b/Engine/Source/Runtime/Engine/Private/WorldPartition/DataLayer/DataLayerSubsystem.cpp
--- a/Engine/Source/Runtime/Engine/Private/WorldPartition/DataLayer/DataLayerSubsystem.cpp
+++ b/Engine/Source/Runtime/Engine/Private/WorldPartition/DataLayer/DataLayerSubsystem.cpp
@@ -499,28 +499,28 @@ const UDataLayerInstance* UDataLayerSubsystem::GetDataLayerInstanceFromAssetName
 	return nullptr;
 }
 
-void UDataLayerSubsystem::ForEachDataLayer(TFunctionRef<bool(UDataLayerInstance*)> Func)
+void UDataLayerSubsystem::ForEachDataLayer(TFunctionRef<bool(UDataLayerInstance*)> Func, const ULevel* InLevelContext /* = nullptr */)
 {
-	if (AWorldDataLayers* WorldDataLayers = GetWorld()->GetWorldDataLayers())
+	if (!InLevelContext || InLevelContext == GetWorld()->PersistentLevel)
 	{
-		WorldDataLayers->ForEachDataLayer(Func);
+		if (AWorldDataLayers* WorldDataLayers = GetWorld()->GetWorldDataLayers())
+		{
+			WorldDataLayers->ForEachDataLayer(Func);
+		}
+		else
+		{
+			UE_LOG(LogWorldPartition, Warning, TEXT("Iterating over Data Layers without a World Partition"));
+		}
 	}
-	else
+	else if (const AWorldDataLayers* CurrentLevelWorldDataLayers = (InLevelContext && !InLevelContext->IsPersistentLevel()) ? InLevelContext->GetWorldDataLayers() : nullptr)
 	{
-		UE_LOG(LogWorldPartition, Warning, TEXT("Iterating over Data Layers without a World Partition"));
+		CurrentLevelWorldDataLayers->ForEachDataLayer(Func);
 	}
 }
 
-void UDataLayerSubsystem::ForEachDataLayer(TFunctionRef<bool(UDataLayerInstance*)> Func) const
+void UDataLayerSubsystem::ForEachDataLayer(TFunctionRef<bool(UDataLayerInstance*)> Func, const ULevel* InLevelContext /* = nullptr */) const
 {
-	if (AWorldDataLayers* WorldDataLayers = GetWorld()->GetWorldDataLayers())
-	{
-		WorldDataLayers->ForEachDataLayer(Func);
-	}
-	else
-	{
-		UE_LOG(LogWorldPartition, Warning, TEXT("Iterating over Data Layers without a World Partition"));
-	}
+	const_cast<UDataLayerSubsystem*>(this)->ForEachDataLayer(Func, InLevelContext);
 }
 
 FAutoConsoleCommand UDataLayerSubsystem::ToggleDataLayerActivation(
diff --git a/Engine/Source/Runtime/Engine/Private/WorldPartition/DataLayer/DataLayerUtils.cpp b/Engine/Source/Runtime/Engine/Private/WorldPartition/DataLayer/DataLayerUtils.cpp
--- a/Engine/Source/Runtime/Engine/Private/WorldPartition/DataLayer/DataLayerUtils.cpp
+++ b/Engine/Source/Runtime/Engine/Private/WorldPartition/DataLayer/DataLayerUtils.cpp
@@ -29,15 +29,19 @@ FWorldDataLayersActorDesc* FDataLayerUtils::GetWorldDataLayersActorDesc(const UA
 	return nullptr;
 }
 
-TArray<FName> FDataLayerUtils::ResolvedDataLayerInstanceNames(const FWorldPartitionActorDesc* InActorDesc, const FWorldDataLayersActorDesc* InWorldDataLayersActorDesc, bool* bOutIsResultValid)
+TArray<FName> FDataLayerUtils::ResolvedDataLayerInstanceNames(const FWorldPartitionActorDesc* InActorDesc, const FWorldDataLayersActorDesc* InWorldDataLayersActorDesc, UWorld* InWorld, bool* bOutIsResultValid)
 {
 	bool bLocalIsSuccess = true;
 	bool& bIsSuccess = bOutIsResultValid ? *bOutIsResultValid : bLocalIsSuccess;
 	bIsSuccess = true;
 
 	// Prioritize in-memory AWorldDataLayers
-	UWorld* World = InActorDesc->GetContainer() ? InActorDesc->GetContainer()->GetWorld() : nullptr;
-	const UDataLayerSubsystem* DataLayerSubsystem = World ? World->GetSubsystem<UDataLayerSubsystem>() : nullptr;
+	UWorld* World = InWorld;
+	if (!World)
+	{
+		World = InActorDesc->GetContainer() ? InActorDesc->GetContainer()->GetWorld() : nullptr;
+	}
+	const UDataLayerSubsystem* DataLayerSubsystem = UWorld::GetSubsystem<UDataLayerSubsystem>(World);
 
 	// DataLayers not using DataLayer Assets represent DataLayerInstanceNames
 	if (!InActorDesc->IsUsingDataLayerAsset())
diff --git a/Engine/Source/Runtime/Engine/Private/WorldPartition/DataLayer/DeprecatedDataLayerInstance.cpp b/Engine/Source/Runtime/Engine/Private/WorldPartition/DataLayer/DeprecatedDataLayerInstance.cpp
--- a/Engine/Source/Runtime/Engine/Private/WorldPartition/DataLayer/DeprecatedDataLayerInstance.cpp
+++ b/Engine/Source/Runtime/Engine/Private/WorldPartition/DataLayer/DeprecatedDataLayerInstance.cpp
@@ -72,11 +72,6 @@ bool UDeprecatedDataLayerInstance::RemoveActor(AActor* Actor) const
 	return Actor->RemoveDataLayer(FActorDataLayer(GetDataLayerFName()));
 }
 
-bool UDeprecatedDataLayerInstance::ContainsActor(const AActor* Actor) const
-{
-	return Actor->ContainsDataLayer(FActorDataLayer(GetDataLayerFName()));
-}
-
 bool UDeprecatedDataLayerInstance::RelabelDataLayer(FName InDataLayerLabel)
 {
 	FName SanitizedLabel = FDataLayerUtils::GetSanitizedDataLayerLabel(InDataLayerLabel);
diff --git a/Engine/Source/Runtime/Engine/Private/WorldPartition/DataLayer/WorldDataLayers.cpp b/Engine/Source/Runtime/Engine/Private/WorldPartition/DataLayer/WorldDataLayers.cpp
--- a/Engine/Source/Runtime/Engine/Private/WorldPartition/DataLayer/WorldDataLayers.cpp
+++ b/Engine/Source/Runtime/Engine/Private/WorldPartition/DataLayer/WorldDataLayers.cpp
@@ -450,6 +450,12 @@ AWorldDataLayers* AWorldDataLayers::Create(UWorld* World, FName InWorldDataLayer
 	return WorldDataLayers;
 }
 
+bool AWorldDataLayers::IsMainWorldDataLayers() const
+{
+	check(GetWorld());
+	return GetWorld()->GetWorldDataLayers() == this;
+}
+
 TArray<FName> AWorldDataLayers::GetDataLayerInstanceNames(const TArray<TObjectPtr<const UDataLayerAsset>>& InDataLayersAssets) const
 {
 	TArray<FName> OutDataLayerNames;
diff --git a/Engine/Source/Runtime/Engine/Private/WorldPartition/WorldPartition.cpp b/Engine/Source/Runtime/Engine/Private/WorldPartition/WorldPartition.cpp
--- a/Engine/Source/Runtime/Engine/Private/WorldPartition/WorldPartition.cpp
+++ b/Engine/Source/Runtime/Engine/Private/WorldPartition/WorldPartition.cpp
@@ -501,13 +501,13 @@ void UWorldPartition::Initialize(UWorld* InWorld, const FTransform& InTransform)
 	if (bIsEditor)
 	{
 		// Repair existing maps with no UWorld::WorldDataLayers actor.
-		AWorldDataLayers* WorldDataLayers = World->GetWorldDataLayers();
+		AWorldDataLayers* WorldDataLayers = OuterWorld->GetWorldDataLayers();
 		if (!WorldDataLayers)
 		{
 			// WorldDataLayersActor is invalid here orelse its PostLoad would have set itself as the world's WorldDataLayers
 			check(!WorldDataLayersActor.IsValid());
-			WorldDataLayers = AWorldDataLayers::Create(World);
-			World->SetWorldDataLayers(WorldDataLayers);
+			WorldDataLayers = AWorldDataLayers::Create(OuterWorld);
+			OuterWorld->SetWorldDataLayers(WorldDataLayers);
 		}
 
 		// Apply level transform on actors already part of the level
@@ -704,32 +704,30 @@ void UWorldPartition::RegisterDelegates()
 {
 	check(World); 
 
-	if (IsMainWorldPartition())
-	{
 #if WITH_EDITOR
-		if (GEditor && !IsTemplate() && !World->IsGameWorld())
+	if (GEditor && !IsTemplate() && !World->IsGameWorld())
+	{
+		if (IsMainWorldPartition())
 		{
 			FEditorDelegates::PreBeginPIE.AddUObject(this, &UWorldPartition::OnPreBeginPIE);
 			FEditorDelegates::PrePIEEnded.AddUObject(this, &UWorldPartition::OnPrePIEEnded);
 			FEditorDelegates::CancelPIE.AddUObject(this, &UWorldPartition::OnCancelPIE);
 			FGameDelegates::Get().GetEndPlayMapDelegate().AddUObject(this, &UWorldPartition::OnEndPlay);
-
 			FCoreUObjectDelegates::PostReachabilityAnalysis.AddUObject(this, &UWorldPartition::OnGCPostReachabilityAnalysis);
-
 			GEditor->OnPostBugItGoCalled().AddUObject(this, &UWorldPartition::OnPostBugItGoCalled);
-
-			FDataLayersEditorBroadcast::Get().OnActorDataLayersEditorLoadingStateChanged().AddUObject(this, &UWorldPartition::OnActorDataLayersEditorLoadingStateChanged);
 		}
+
+		FDataLayersEditorBroadcast::Get().OnActorDataLayersEditorLoadingStateChanged().AddUObject(this, &UWorldPartition::OnActorDataLayersEditorLoadingStateChanged);
+	}
 #endif
 
-		if (World->IsGameWorld())
-		{
-			World->OnWorldMatchStarting.AddUObject(this, &UWorldPartition::OnWorldMatchStarting);
+	if (IsMainWorldPartition() && World->IsGameWorld())
+	{
+		World->OnWorldMatchStarting.AddUObject(this, &UWorldPartition::OnWorldMatchStarting);
 
 #if !UE_BUILD_SHIPPING
-			FCoreDelegates::OnGetOnScreenMessages.AddUObject(this, &UWorldPartition::GetOnScreenMessages);
+		FCoreDelegates::OnGetOnScreenMessages.AddUObject(this, &UWorldPartition::GetOnScreenMessages);
 #endif
-		}
 	}
 }
 
@@ -737,10 +735,10 @@ void UWorldPartition::UnregisterDelegates()
 {
 	check(World);
 
-	if (IsMainWorldPartition())
-	{
 #if WITH_EDITOR
-		if (GEditor && !IsTemplate() && !World->IsGameWorld())
+	if (GEditor && !IsTemplate() && !World->IsGameWorld())
+	{
+		if (IsMainWorldPartition())
 		{
 			FEditorDelegates::PreBeginPIE.RemoveAll(this);
 			FEditorDelegates::PrePIEEnded.RemoveAll(this);
@@ -753,19 +751,19 @@ void UWorldPartition::UnregisterDelegates()
 			}
 
 			GEditor->OnPostBugItGoCalled().RemoveAll(this);
-
-			FDataLayersEditorBroadcast::Get().OnActorDataLayersEditorLoadingStateChanged().RemoveAll(this);
 		}
+
+		FDataLayersEditorBroadcast::Get().OnActorDataLayersEditorLoadingStateChanged().RemoveAll(this);
+	}
 #endif
 
-		if (World->IsGameWorld())
-		{
-			World->OnWorldMatchStarting.RemoveAll(this);
+	if (IsMainWorldPartition() && World->IsGameWorld())
+	{
+		World->OnWorldMatchStarting.RemoveAll(this);
 
 #if !UE_BUILD_SHIPPING
-			FCoreDelegates::OnGetOnScreenMessages.RemoveAll(this);
+		FCoreDelegates::OnGetOnScreenMessages.RemoveAll(this);
 #endif
-		}
 	}
 }
 
@@ -1417,7 +1415,7 @@ bool UWorldPartition::IsStreamingCompleted(EWorldPartitionRuntimeCellState Query
 
 bool UWorldPartition::CanDrawRuntimeHash() const
 {
-	return GetWorld()->IsGameWorld() || UWorldPartition::IsSimulating();
+	return GetWorld()->IsGameWorld() || (UWorldPartition::IsSimulating() && StreamingPolicy);
 }
 
 FVector2D UWorldPartition::GetDrawRuntimeHash2DDesiredFootprint(const FVector2D& CanvasSize)
diff --git a/Engine/Source/Runtime/Engine/Private/WorldPartition/WorldPartitionActorDesc.cpp b/Engine/Source/Runtime/Engine/Private/WorldPartition/WorldPartitionActorDesc.cpp
--- a/Engine/Source/Runtime/Engine/Private/WorldPartition/WorldPartitionActorDesc.cpp
+++ b/Engine/Source/Runtime/Engine/Private/WorldPartition/WorldPartitionActorDesc.cpp
@@ -17,6 +17,7 @@
 #include "WorldPartition/WorldPartitionLog.h"
 #include "WorldPartition/ActorDescContainer.h"
 #include "WorldPartition/HLOD/HLODLayer.h"
+#include "WorldPartition/DataLayer/WorldDataLayers.h"
 #include "WorldPartition/DataLayer/DataLayerSubsystem.h"
 #include "WorldPartition/DataLayer/DataLayerAsset.h"
 #include "WorldPartition/DataLayer/DataLayerUtils.h"
@@ -69,20 +70,23 @@ void FWorldPartitionActorDesc::Init(const AActor* InActor)
 		TArray<FName> LocalDataLayerAssetPaths;
 		TArray<FName> LocalDataLayerInstanceNames;
 		UDataLayerSubsystem* DataLayerSubsystem = UWorld::GetSubsystem<UDataLayerSubsystem>(InActor->GetWorld());
-		
+		ULevel* Level = InActor->GetLevel();
+
 		LocalDataLayerAssetPaths.Reserve(InActor->GetDataLayerAssets().Num());
 		for (const TObjectPtr<const UDataLayerAsset>& DataLayerAsset : InActor->GetDataLayerAssets())
 		{
-			if (DataLayerAsset && DataLayerSubsystem->GetDataLayerInstance(DataLayerAsset))
+			// Pass Actor Level when resolving the DataLayerInstance as FWorldPartitionActorDesc always represents the state of the actor local to its outer level
+			if (DataLayerAsset && DataLayerSubsystem->GetDataLayerInstance(DataLayerAsset, Level))
 			{
 				LocalDataLayerAssetPaths.Add(*DataLayerAsset->GetPathName());
 			}
 		}
 
 		PRAGMA_DISABLE_DEPRECATION_WARNINGS
-		LocalDataLayerInstanceNames = DataLayerSubsystem->GetDataLayerInstanceNames(InActor->GetActorDataLayers());
+		// Pass Actor Level when resolving the DataLayerInstance as FWorldPartitionActorDesc always represents the state of the actor local to its outer level
+		LocalDataLayerInstanceNames = DataLayerSubsystem->GetDataLayerInstanceNames(InActor->GetActorDataLayers(), Level);
 		PRAGMA_ENABLE_DEPRECATION_WARNINGS
-		
+
 		// Validation
 		const bool bHasDataLayerAssets = LocalDataLayerAssetPaths.Num() > 0;
 		const bool bHasDeprecatedDataLayers = LocalDataLayerInstanceNames.Num() > 0;
@@ -93,7 +97,7 @@ void FWorldPartitionActorDesc::Init(const AActor* InActor)
 		DataLayers = bIsUsingDataLayerAsset ? MoveTemp(LocalDataLayerAssetPaths) : MoveTemp(LocalDataLayerInstanceNames);
 
 		// Init DataLayers transient info
-		DataLayerInstanceNames = FDataLayerUtils::ResolvedDataLayerInstanceNames(this);
+		DataLayerInstanceNames = FDataLayerUtils::ResolvedDataLayerInstanceNames(this, nullptr, InActor->GetWorld());
 	}
 
 	Tags = InActor->Tags;
diff --git a/Engine/Source/Runtime/Engine/Private/WorldPartition/WorldPartitionStreamingGeneration.cpp b/Engine/Source/Runtime/Engine/Private/WorldPartition/WorldPartitionStreamingGeneration.cpp
--- a/Engine/Source/Runtime/Engine/Private/WorldPartition/WorldPartitionStreamingGeneration.cpp
+++ b/Engine/Source/Runtime/Engine/Private/WorldPartition/WorldPartitionStreamingGeneration.cpp
@@ -132,17 +132,11 @@ class FWorldPartitionStreamingGenerator
 				const FGuid ActorGuid = ActorDescView.GetGuid();
 				const FActorContainerID SubContainerID(InContainerID, ActorGuid);
 
-				// Only propagate data layers from the root container as we don't support sub containers to set their own
-				const TSet<FName>* SubRuntimeDataLayers = &InRuntimeDataLayers;
-				TSet<FName> ParentRuntimeDataLayers;
+				// Combine actor runtime Data Layers with parent container runtime Data Layers
+				TSet<FName> CombinedRuntimeDataLayers = InRuntimeDataLayers;
+				CombinedRuntimeDataLayers.Append(ActorDescView.GetRuntimeDataLayers());
 
-				if (InContainerID.IsMainContainer())
-				{
-					ParentRuntimeDataLayers.Append(ActorDescView.GetRuntimeDataLayers());
-					SubRuntimeDataLayers = &ParentRuntimeDataLayers;
-				}
-
-				CreateActorDescriptorViewsRecursive(SubContainer, SubTransform * InTransform, *SubRuntimeDataLayers, SubContainerID, InContainerID, SubClusterMode, *ActorDescView.GetActorLabelOrName().ToString());
+				CreateActorDescriptorViewsRecursive(SubContainer, SubTransform * InTransform, CombinedRuntimeDataLayers, SubContainerID, InContainerID, SubClusterMode, *ActorDescView.GetActorLabelOrName().ToString());
 
 				// The container actor can be removed now that its contained actors has been registered
 				It.RemoveCurrent();
diff --git a/Engine/Source/Runtime/Engine/Private/WorldPartition/WorldPartitionSubsystem.cpp b/Engine/Source/Runtime/Engine/Private/WorldPartition/WorldPartitionSubsystem.cpp
--- a/Engine/Source/Runtime/Engine/Private/WorldPartition/WorldPartitionSubsystem.cpp
+++ b/Engine/Source/Runtime/Engine/Private/WorldPartition/WorldPartitionSubsystem.cpp
@@ -159,6 +159,19 @@ void UWorldPartitionSubsystem::Deinitialize()
 	Super::Deinitialize();
 }
 
+#if WITH_EDITOR
+void UWorldPartitionSubsystem::ForEachWorldPartition(TFunctionRef<bool(UWorldPartition*)> Func)
+{
+	for (UWorldPartition* WorldPartition : RegisteredWorldPartitions)
+	{
+		if (!Func(WorldPartition))
+		{
+			return;
+		}
+	}
+}
+#endif
+
 void UWorldPartitionSubsystem::OnWorldPartitionInitialized(UWorldPartition* InWorldPartition)
 {
 	if (RegisteredWorldPartitions.IsEmpty())
diff --git a/Engine/Source/Runtime/Engine/Public/LevelInstance/LevelInstanceEditorLevelStreaming.h b/Engine/Source/Runtime/Engine/Public/LevelInstance/LevelInstanceEditorLevelStreaming.h
--- a/Engine/Source/Runtime/Engine/Public/LevelInstance/LevelInstanceEditorLevelStreaming.h
+++ b/Engine/Source/Runtime/Engine/Public/LevelInstance/LevelInstanceEditorLevelStreaming.h
@@ -25,6 +25,7 @@ public:
 	virtual TOptional<FFolder::FRootObject> GetFolderRootObject() const override;
 protected:
 	void OnLevelActorAdded(AActor* InActor);
+	void OnLoadedActorAddedToLevel(AActor& InActor);
 
 	friend class ULevelInstanceSubsystem;
 
diff --git a/Engine/Source/Runtime/Engine/Public/LevelInstance/LevelInstanceLevelStreaming.h b/Engine/Source/Runtime/Engine/Public/LevelInstance/LevelInstanceLevelStreaming.h
--- a/Engine/Source/Runtime/Engine/Public/LevelInstance/LevelInstanceLevelStreaming.h
+++ b/Engine/Source/Runtime/Engine/Public/LevelInstance/LevelInstanceLevelStreaming.h
@@ -9,6 +9,7 @@
 #include "LevelInstanceLevelStreaming.generated.h"
 
 class ILevelInstanceInterface;
+class ALevelInstanceEditorInstanceActor;
 
 UCLASS(Transient)
 class ENGINE_API ULevelStreamingLevelInstance : public ULevelStreamingDynamic
@@ -34,5 +35,12 @@ protected:
 	friend class ULevelInstanceSubsystem;
 
 private:
+#if WITH_EDITOR
+	void PrepareLevelInstanceLoadedActor(AActor& InActor, ILevelInstanceInterface* InLevelInstance);
+	void OnLoadedActorAddedToLevel(AActor& InActor);
+	void OnLoadedActorRemovedFromLevel(AActor& InActor);
+
+	TWeakObjectPtr<ALevelInstanceEditorInstanceActor> LevelInstanceEditorInstanceActor;
+#endif
 	FLevelInstanceID LevelInstanceID;
 };
\ No newline at end of file
diff --git a/Engine/Source/Runtime/Engine/Public/LevelInstance/LevelInstanceSubsystem.h b/Engine/Source/Runtime/Engine/Public/LevelInstance/LevelInstanceSubsystem.h
--- a/Engine/Source/Runtime/Engine/Public/LevelInstance/LevelInstanceSubsystem.h
+++ b/Engine/Source/Runtime/Engine/Public/LevelInstance/LevelInstanceSubsystem.h
@@ -40,6 +40,7 @@ public:
 	//~ End UWorldSubsystem Interface.
 
 	ILevelInstanceInterface* GetLevelInstance(FLevelInstanceID LevelInstanceID) const;
+	ILevelInstanceInterface* GetOwningLevelInstance(const ULevel* Level) const;
 	FLevelInstanceID RegisterLevelInstance(ILevelInstanceInterface* LevelInstance);
 	void UnregisterLevelInstance(ILevelInstanceInterface* LevelInstance);
 	void RequestLoadLevelInstance(ILevelInstanceInterface* LevelInstance, bool bUpdate);
@@ -97,6 +98,11 @@ public:
 
 	TArray<ILevelInstanceInterface*> GetLevelInstances(const FString& WorldAssetPackage);
 	
+	// Returns the upper chain of level instance actors for the specified level starting with the level instance referencing the level
+	void ForEachLevelInstanceActorAncestors(const ULevel* Level, TFunctionRef<bool(AActor*)> Operation) const;
+	TArray<AActor*> GetParentLevelInstanceActors(const ULevel* Level) const;
+	FString PrefixWithParentLevelInstanceActorLabels(const FString& ActorLabel, const ULevel* Level) const;
+
 	static bool CheckForLoop(const ILevelInstanceInterface* LevelInstance, TArray<TPair<FText, TSoftObjectPtr<UWorld>>>* LoopInfo = nullptr, const ILevelInstanceInterface** LoopStart = nullptr);
 	static bool CheckForLoop(const ILevelInstanceInterface* LevelInstance, TSoftObjectPtr<UWorld> WorldAsset, TArray<TPair<FText, TSoftObjectPtr<UWorld>>>* LoopInfo = nullptr, const ILevelInstanceInterface** LoopStart = nullptr);
 	static bool CanUseWorldAsset(const ILevelInstanceInterface* LevelInstance, TSoftObjectPtr<UWorld> WorldAsset, FString* OutReason);
@@ -109,7 +115,6 @@ private:
 	void UnloadLevelInstance(const FLevelInstanceID& LevelInstanceID);
 	void ForEachActorInLevel(ULevel* Level, TFunctionRef<bool(AActor * LevelActor)> Operation) const;
 	void ForEachLevelInstanceAncestors(AActor* Actor, TFunctionRef<bool(ILevelInstanceInterface*)> Operation) const;
-	ILevelInstanceInterface* GetOwningLevelInstance(const ULevel* Level) const;
 	
 	void RegisterLoadedLevelStreamingLevelInstance(ULevelStreamingLevelInstance* LevelStreaming);
 
diff --git a/Engine/Source/Runtime/Engine/Public/WorldPartition/ActorDescContainer.h b/Engine/Source/Runtime/Engine/Public/WorldPartition/ActorDescContainer.h
--- a/Engine/Source/Runtime/Engine/Public/WorldPartition/ActorDescContainer.h
+++ b/Engine/Source/Runtime/Engine/Public/WorldPartition/ActorDescContainer.h
@@ -84,6 +84,7 @@ protected:
 
 private:
 #if WITH_EDITOR
+	bool ShouldRegisterDelegates();
 	void RegisterEditorDelegates();
 	void UnregisterEditorDelegates();
 #endif
diff --git a/Engine/Source/Runtime/Engine/Public/WorldPartition/DataLayer/DataLayerInstance.h b/Engine/Source/Runtime/Engine/Public/WorldPartition/DataLayer/DataLayerInstance.h
--- a/Engine/Source/Runtime/Engine/Public/WorldPartition/DataLayer/DataLayerInstance.h
+++ b/Engine/Source/Runtime/Engine/Public/WorldPartition/DataLayer/DataLayerInstance.h
@@ -50,7 +50,6 @@ public:
 	void SetIsLoadedInEditor(bool bIsLoadedInEditor, bool bFromUserChange);
 	void SetIsLocked(bool bInIsLocked) { bIsLocked = bInIsLocked; }
 
-	bool IsLocked() const;
 	bool IsInitiallyLoadedInEditor() const { return bIsInitiallyLoadedInEditor; }
 	bool IsLoadedInEditor() const { return bIsLoadedInEditor; }
 	bool IsEffectiveLoadedInEditor() const;
@@ -67,9 +66,11 @@ public:
 
 	static FText GetDataLayerText(const UDataLayerInstance* InDataLayer);
 
+	virtual bool IsLocked() const;
+	virtual bool IsReadOnly() const { return false; }
+	virtual bool CanEditChange(const FProperty* InProperty) const;
 	virtual bool AddActor(AActor* Actor) const { return false; }
 	virtual bool RemoveActor(AActor* Actor) const { return false; }
-	virtual bool ContainsActor(const AActor* Actor) const { return false; }
 
 	bool IsInActorEditorContext() const;
 	bool AddToActorEditorContext();
diff --git a/Engine/Source/Runtime/Engine/Public/WorldPartition/DataLayer/DataLayerInstanceWithAsset.h b/Engine/Source/Runtime/Engine/Public/WorldPartition/DataLayer/DataLayerInstanceWithAsset.h
--- a/Engine/Source/Runtime/Engine/Public/WorldPartition/DataLayer/DataLayerInstanceWithAsset.h
+++ b/Engine/Source/Runtime/Engine/Public/WorldPartition/DataLayer/DataLayerInstanceWithAsset.h
@@ -18,9 +18,10 @@ public:
 	static FName MakeName(const UDataLayerAsset* DeprecatedDataLayer);
 	void OnCreated(const UDataLayerAsset* Asset);
 
+	virtual bool IsLocked() const override;
+	virtual bool IsReadOnly() const override;
 	virtual bool AddActor(AActor* Actor) const override;
 	virtual bool RemoveActor(AActor* Actor) const override;
-	virtual bool ContainsActor(const AActor* Actor) const override;
 
 	virtual bool Validate(IStreamingGenerationErrorHandler* ErrorHandler) const override;
 #endif
diff --git a/Engine/Source/Runtime/Engine/Public/WorldPartition/DataLayer/DataLayerSubsystem.h b/Engine/Source/Runtime/Engine/Public/WorldPartition/DataLayer/DataLayerSubsystem.h
--- a/Engine/Source/Runtime/Engine/Public/WorldPartition/DataLayer/DataLayerSubsystem.h
+++ b/Engine/Source/Runtime/Engine/Public/WorldPartition/DataLayer/DataLayerSubsystem.h
@@ -92,18 +92,18 @@ public:
 #endif
 
 	template<class T>
-	UDataLayerInstance* GetDataLayerInstance(const T& InDataLayerIdentifier) const;
+	UDataLayerInstance* GetDataLayerInstance(const T& InDataLayerIdentifier, const ULevel* InLevelContext = nullptr) const;
 
 	template<class T>
-	TArray<const UDataLayerInstance*> GetDataLayerInstances(const TArray<T>& InDataLayerIdentifiers) const;
+	TArray<const UDataLayerInstance*> GetDataLayerInstances(const TArray<T>& InDataLayerIdentifiers, const ULevel* InLevelContext = nullptr) const;
 
 	template<class T>
-	TArray<FName> GetDataLayerInstanceNames(const TArray<T>& InDataLayerIdentifiers) const;
+	TArray<FName> GetDataLayerInstanceNames(const TArray<T>& InDataLayerIdentifiers, const ULevel* InLevelContext = nullptr) const;
 
 	const UDataLayerInstance* GetDataLayerInstanceFromAssetName(const FName& InDataLayerAssetFullName) const;
 
-	void ForEachDataLayer(TFunctionRef<bool(UDataLayerInstance*)> Func);
-	void ForEachDataLayer(TFunctionRef<bool(UDataLayerInstance*)> Func) const;
+	void ForEachDataLayer(TFunctionRef<bool(UDataLayerInstance*)> Func, const ULevel* InLevelContext = nullptr);
+	void ForEachDataLayer(TFunctionRef<bool(UDataLayerInstance*)> Func, const ULevel* InLevelContext = nullptr) const;
 
 	void SetDataLayerRuntimeState(const UDataLayerInstance* InDataLayerInstance, EDataLayerRuntimeState InState, bool bInIsRecursive = false);
 	EDataLayerRuntimeState GetDataLayerRuntimeState(const UDataLayerInstance* InDataLayer) const;
@@ -213,47 +213,37 @@ private:
 };
 
 template<class T>
-UDataLayerInstance* UDataLayerSubsystem::GetDataLayerInstance(const T& InDataLayerIdentifier) const
+UDataLayerInstance* UDataLayerSubsystem::GetDataLayerInstance(const T& InDataLayerIdentifier, const ULevel* InLevelContext /* = nullptr */) const
 {
-	const UDataLayerInstance* DataLayerInstance = nullptr;
-	if (AWorldDataLayers* WorldDataLayers = GetWorld()->GetWorldDataLayers())
-	{
-		PRAGMA_DISABLE_DEPRECATION_WARNINGS
-			DataLayerInstance = WorldDataLayers->GetDataLayerInstance(InDataLayerIdentifier);
-		PRAGMA_ENABLE_DEPRECATION_WARNINGS
-	}
-	
+	PRAGMA_DISABLE_DEPRECATION_WARNINGS
+	const AWorldDataLayers* WorldDataLayers = InLevelContext ? InLevelContext->GetWorldDataLayers() : GetWorld()->GetWorldDataLayers();
+	const UDataLayerInstance * DataLayerInstance = WorldDataLayers ? WorldDataLayers->GetDataLayerInstance(InDataLayerIdentifier) : nullptr;
+	PRAGMA_ENABLE_DEPRECATION_WARNINGS
 	return const_cast<UDataLayerInstance*>(DataLayerInstance);
 }
 
 template<class T>
-TArray<FName> UDataLayerSubsystem::GetDataLayerInstanceNames(const TArray<T>& InDataLayerIdentifiers) const
+TArray<FName> UDataLayerSubsystem::GetDataLayerInstanceNames(const TArray<T>& InDataLayerIdentifiers, const ULevel* InLevelContext /* = nullptr */) const
 {
 	TArray<FName> DataLayerInstanceNames;
 	DataLayerInstanceNames.Reserve(InDataLayerIdentifiers.Num());
 
-	if (AWorldDataLayers* WorldDataLayers = GetWorld()->GetWorldDataLayers())
-	{
-		PRAGMA_DISABLE_DEPRECATION_WARNINGS
-			DataLayerInstanceNames = WorldDataLayers->GetDataLayerInstanceNames(InDataLayerIdentifiers);
-		PRAGMA_ENABLE_DEPRECATION_WARNINGS
-	}
-
+	PRAGMA_DISABLE_DEPRECATION_WARNINGS
+	const AWorldDataLayers* WorldDataLayers = InLevelContext ? InLevelContext->GetWorldDataLayers() : GetWorld()->GetWorldDataLayers();
+	DataLayerInstanceNames = WorldDataLayers->GetDataLayerInstanceNames(InDataLayerIdentifiers);
+	PRAGMA_ENABLE_DEPRECATION_WARNINGS
 	return DataLayerInstanceNames;
 }
 
 template<class T>
-TArray<const UDataLayerInstance*> UDataLayerSubsystem::GetDataLayerInstances(const TArray<T>& InDataLayerIdentifiers) const
+TArray<const UDataLayerInstance*> UDataLayerSubsystem::GetDataLayerInstances(const TArray<T>& InDataLayerIdentifiers, const ULevel* InLevelContext /* = nullptr */) const
 {
 	TArray<const UDataLayerInstance*> DataLayerInstances;
 	DataLayerInstances.Reserve(InDataLayerIdentifiers.Num());
 
-	if (AWorldDataLayers* WorldDataLayers = GetWorld()->GetWorldDataLayers())
-	{
-		PRAGMA_DISABLE_DEPRECATION_WARNINGS
-			DataLayerInstances = WorldDataLayers->GetDataLayerInstances(InDataLayerIdentifiers);
-		PRAGMA_ENABLE_DEPRECATION_WARNINGS
-	}
-
+	PRAGMA_DISABLE_DEPRECATION_WARNINGS
+	const AWorldDataLayers* WorldDataLayers = InLevelContext ? InLevelContext->GetWorldDataLayers() : GetWorld()->GetWorldDataLayers();
+	DataLayerInstances = WorldDataLayers->GetDataLayerInstances(InDataLayerIdentifiers);
+	PRAGMA_ENABLE_DEPRECATION_WARNINGS
 	return DataLayerInstances;
 }
\ No newline at end of file
diff --git a/Engine/Source/Runtime/Engine/Public/WorldPartition/DataLayer/DataLayerUtils.h b/Engine/Source/Runtime/Engine/Public/WorldPartition/DataLayer/DataLayerUtils.h
--- a/Engine/Source/Runtime/Engine/Public/WorldPartition/DataLayer/DataLayerUtils.h
+++ b/Engine/Source/Runtime/Engine/Public/WorldPartition/DataLayer/DataLayerUtils.h
@@ -23,7 +23,7 @@ public:
 		return IconNameByType[static_cast<uint32>(DataLayerType)];
 	}
 	static FWorldDataLayersActorDesc* GetWorldDataLayersActorDesc(const UActorDescContainer* InContainer, bool bInCheckValid = true);
-	static TArray<FName> ResolvedDataLayerInstanceNames(const FWorldPartitionActorDesc* InActorDesc, const FWorldDataLayersActorDesc* InWorldDataLayersActorDesc = nullptr, bool* bOutIsResultValid = nullptr);
+	static TArray<FName> ResolvedDataLayerInstanceNames(const FWorldPartitionActorDesc* InActorDesc, const FWorldDataLayersActorDesc* InWorldDataLayersActorDesc = nullptr, UWorld* InWorld = nullptr, bool* bOutIsResultValid = nullptr);
 	static TArray<FName> ResolveRuntimeDataLayerInstanceNames(const FWorldPartitionActorDesc* InActorDesc, const UActorDescContainer* InContainer, bool* bOutIsResultValid = nullptr);
 #endif
 
diff --git a/Engine/Source/Runtime/Engine/Public/WorldPartition/DataLayer/DeprecatedDataLayerInstance.h b/Engine/Source/Runtime/Engine/Public/WorldPartition/DataLayer/DeprecatedDataLayerInstance.h
--- a/Engine/Source/Runtime/Engine/Public/WorldPartition/DataLayer/DeprecatedDataLayerInstance.h
+++ b/Engine/Source/Runtime/Engine/Public/WorldPartition/DataLayer/DeprecatedDataLayerInstance.h
@@ -41,7 +41,6 @@ private:
 #if WITH_EDITOR
 	virtual bool AddActor(AActor* Actor) const override;
 	virtual bool RemoveActor(AActor* Actor) const override;
-	virtual bool ContainsActor(const AActor* Actor) const override;
 
 	virtual bool SupportRelabeling() const override { return true; }
 	virtual bool RelabelDataLayer(FName NewDataLayerLabel) override;
diff --git a/Engine/Source/Runtime/Engine/Public/WorldPartition/DataLayer/WorldDataLayers.h b/Engine/Source/Runtime/Engine/Public/WorldPartition/DataLayer/WorldDataLayers.h
--- a/Engine/Source/Runtime/Engine/Public/WorldPartition/DataLayer/WorldDataLayers.h
+++ b/Engine/Source/Runtime/Engine/Public/WorldPartition/DataLayer/WorldDataLayers.h
@@ -33,7 +33,7 @@ struct FActorPlacementDataLayers
 /**
  * Actor containing all data layers for a world
  */
-UCLASS(hidecategories = (Actor, HLOD, Cooking, Transform, Advanced, Display, Events, Object, Attachment, Info, Input, Blueprint, Layers, Tags, Replication), notplaceable)
+UCLASS(hidecategories = (Actor, HLOD, Cooking, Transform, Advanced, Display, Events, Object, Physics, Attachment, Info, Input, Blueprint, Layers, Tags, Replication), notplaceable)
 class ENGINE_API AWorldDataLayers : public AInfo
 {
 	GENERATED_UCLASS_BODY()
@@ -80,6 +80,8 @@ public:
 	// Allows overriding of DataLayers with PlayFromHere
 	template<class T>
 	void OverwriteDataLayerRuntimeStates(const TArray<T>* InActiveDataLayers, const TArray<T>* InLoadedDataLayers );
+
+	bool IsMainWorldDataLayers() const;
 #endif
 	
 	void DumpDataLayers(FOutputDevice& OutputDevice) const;
diff --git a/Engine/Source/Runtime/Engine/Public/WorldPartition/WorldPartitionEditorPerProjectUserSettings.h b/Engine/Source/Runtime/Engine/Public/WorldPartition/WorldPartitionEditorPerProjectUserSettings.h
--- a/Engine/Source/Runtime/Engine/Public/WorldPartition/WorldPartitionEditorPerProjectUserSettings.h
+++ b/Engine/Source/Runtime/Engine/Public/WorldPartition/WorldPartitionEditorPerProjectUserSettings.h
@@ -61,6 +61,7 @@ public:
 		, bHideUnloadedActors(false)
 		, bShowOnlySelectedActors(false)
 		, bHighlightSelectedDataLayers(true)
+		, bHideLevelInstanceContent(true)
 		, bDisableLoadingOfLastLoadedCells(false)
 #endif
 	{}
@@ -139,6 +140,10 @@ public:
 	UPROPERTY(config)
 	uint32 bHighlightSelectedDataLayers : 1;
 
+	/** True when the Data Layer Outliner is not displaying Level Instance content */
+	UPROPERTY(config)
+	uint32 bHideLevelInstanceContent : 1;
+
 private:
 	bool ShouldSaveSettings(const UWorld* InWorld) const
 	{
diff --git a/Engine/Source/Runtime/Engine/Public/WorldPartition/WorldPartitionSubsystem.h b/Engine/Source/Runtime/Engine/Public/WorldPartition/WorldPartitionSubsystem.h
--- a/Engine/Source/Runtime/Engine/Public/WorldPartition/WorldPartitionSubsystem.h
+++ b/Engine/Source/Runtime/Engine/Public/WorldPartition/WorldPartitionSubsystem.h
@@ -47,6 +47,10 @@ public:
 	void RegisterStreamingSourceProvider(IWorldPartitionStreamingSourceProvider* StreamingSource);
 	bool UnregisterStreamingSourceProvider(IWorldPartitionStreamingSourceProvider* StreamingSource);
 
+#if WITH_EDITOR
+	void ForEachWorldPartition(TFunctionRef<bool(UWorldPartition*)> Func);
+#endif
+
 private:
 #if WITH_EDITOR
 	bool IsRunningConvertWorldPartitionCommandlet() const;
