Fixed assert in UDataLayerEditorSubsystem happening in Simulate.

diff --git a/Engine/Source/Editor/DataLayerEditor/Private/DataLayer/DataLayerEditorSubsystem.cpp b/Engine/Source/Editor/DataLayerEditor/Private/DataLayer/DataLayerEditorSubsystem.cpp
--- a/Engine/Source/Editor/DataLayerEditor/Private/DataLayer/DataLayerEditorSubsystem.cpp
+++ b/Engine/Source/Editor/DataLayerEditor/Private/DataLayer/DataLayerEditorSubsystem.cpp
@@ -336,7 +336,7 @@ void UDataLayerEditorSubsystem::InitializeNewActorDataLayers(AActor* Actor)
 
 UWorld* UDataLayerEditorSubsystem::GetWorld() const
 {
-	return GWorld;
+	return GEditor->GetEditorWorldContext().World();
 }
 
 bool UDataLayerEditorSubsystem::SetParentDataLayer(UDataLayerInstance* DataLayer, UDataLayerInstance* ParentDataLayer)
