Data Layers: replaced FScopedDataLayerTransaction by FScopedTransaction.

#rnx

diff --git a/Engine/Source/Editor/DataLayerEditor/Private/DataLayer/DataLayerMode.cpp b/Engine/Source/Editor/DataLayerEditor/Private/DataLayer/DataLayerMode.cpp
--- a/Engine/Source/Editor/DataLayerEditor/Private/DataLayer/DataLayerMode.cpp
+++ b/Engine/Source/Editor/DataLayerEditor/Private/DataLayer/DataLayerMode.cpp
@@ -22,7 +22,7 @@
 #include "ActorDescTreeItem.h"
 #include "ISceneOutlinerHierarchy.h"
 #include "SceneOutlinerMenuContext.h"
-#include "DataLayerTransaction.h"
+#include "ScopedTransaction.h"
 #include "DragAndDrop/ActorDragDropOp.h"
 #include "DragAndDrop/FolderDragDropOp.h"
 #include "ActorMode.h"
@@ -352,7 +352,7 @@ void FDataLayerMode::OnItemDoubleClick(FSceneOutlinerTreeItemPtr Item)
 	{
 		if (UDataLayerInstance* DataLayerInstance = DataLayerItem->GetDataLayer())
 		{
-			const FScopedDataLayerTransaction Transaction(LOCTEXT("SelectActorsInDataLayer", "Select Actors in Data Layer"), GetOwningWorld());
+			const FScopedTransaction Transaction(LOCTEXT("SelectActorsInDataLayer", "Select Actors in Data Layer"));
 			GEditor->SelectNone(/*bNoteSelectionChange*/false, true);
 			DataLayerEditorSubsystem->SelectActorsInDataLayer(DataLayerInstance, /*bSelect*/true, /*bNotify*/true, /*bSelectEvenIfHidden*/true);
 		}
@@ -361,7 +361,7 @@ void FDataLayerMode::OnItemDoubleClick(FSceneOutlinerTreeItemPtr Item)
 	{
 		if (AActor * Actor = DataLayerActorItem->GetActor())
 		{
-			const FScopedDataLayerTransaction Transaction(LOCTEXT("ClickingOnActor", "Clicking on Actor in Data Layer"), GetOwningWorld());
+			const FScopedTransaction Transaction(LOCTEXT("ClickingOnActor", "Clicking on Actor in Data Layer"));
 			GEditor->GetSelectedActors()->Modify();
 			GEditor->SelectNone(/*bNoteSelectionChange*/false, true);
 			GEditor->SelectActor(Actor, /*bSelected*/true, /*bNotify*/true, /*bSelectEvenIfHidden*/true);
@@ -401,7 +401,7 @@ void FDataLayerMode::DeleteItems(const TArray<TWeakPtr<ISceneOutlinerTreeItem>>&
 
 	if (!ActorsToRemoveFromDataLayer.IsEmpty())
 	{
-		const FScopedDataLayerTransaction Transaction(LOCTEXT("RemoveActorsFromDataLayer", "Remove Actors from Data Layer"), GetOwningWorld());
+		const FScopedTransaction Transaction(LOCTEXT("RemoveActorsFromDataLayer", "Remove Actors from Data Layer"));
 		for (const auto& ItPair : ActorsToRemoveFromDataLayer)
 		{
 			DataLayerEditorSubsystem->RemoveActorsFromDataLayer(ItPair.Value, ItPair.Key);
@@ -416,7 +416,7 @@ void FDataLayerMode::DeleteItems(const TArray<TWeakPtr<ISceneOutlinerTreeItem>>&
 		}
 		
 		{
-			const FScopedDataLayerTransaction Transaction(LOCTEXT("DeleteDataLayers", "Delete Data Layers"), GetOwningWorld());
+			const FScopedTransaction Transaction(LOCTEXT("DeleteDataLayers", "Delete Data Layers"));
 			DataLayerEditorSubsystem->DeleteDataLayers(DataLayersToDelete);
 		}
 
@@ -710,7 +710,7 @@ void FDataLayerMode::OnDrop(ISceneOutlinerTreeItem& DropTarget, const FSceneOutl
 			TArray<UDataLayerInstance*> AllSelectedDataLayers = GetSelectedDataLayers(SceneOutliner);
 			if (AllSelectedDataLayers.Num() > 1)
 			{
-				const FScopedDataLayerTransaction Transaction(LOCTEXT("DataLayerOutlinerAddActorsToDataLayers", "Add Actors to Data Layers"), GetOwningWorld());
+				const FScopedTransaction Transaction(LOCTEXT("DataLayerOutlinerAddActorsToDataLayers", "Add Actors to Data Layers"));
 				DataLayerEditorSubsystem->AddActorsToDataLayers(ActorsToAdd, AllSelectedDataLayers);
 				return;
 			}
@@ -718,7 +718,7 @@ void FDataLayerMode::OnDrop(ISceneOutlinerTreeItem& DropTarget, const FSceneOutl
 
 		if (TargetDataLayer)
 		{
-			const FScopedDataLayerTransaction Transaction(LOCTEXT("DataLayerOutlinerAddActorsToDataLayer", "Add Actors to Data Layer"), GetOwningWorld());
+			const FScopedTransaction Transaction(LOCTEXT("DataLayerOutlinerAddActorsToDataLayer", "Add Actors to Data Layer"));
 			DataLayerEditorSubsystem->AddActorsToDataLayer(ActorsToAdd, TargetDataLayer);
 		}
 	}
@@ -727,7 +727,7 @@ void FDataLayerMode::OnDrop(ISceneOutlinerTreeItem& DropTarget, const FSceneOutl
 		TArray<FDataLayerActorMoveElement> ActorsToMove = GetDataLayerActorPairsFromOperation(Payload.SourceOperation);
 		if (!ActorsToMove.IsEmpty() && TargetDataLayer)
 		{
-			const FScopedDataLayerTransaction Transaction(LOCTEXT("DataLayerOutlinerMoveActorsToDataLayer", "Move Actors to Data Layer"), GetOwningWorld());
+			const FScopedTransaction Transaction(LOCTEXT("DataLayerOutlinerMoveActorsToDataLayer", "Move Actors to Data Layer"));
 			for (const auto& DataLayerActorPair : ActorsToMove)
 			{
 				if (AActor* ActorPtr = DataLayerActorPair.Key.Get())
@@ -784,7 +784,7 @@ void FDataLayerMode::SetParentDataLayer(const TArray<UDataLayerInstance*> DataLa
 
 		if (!ValidDataLayers.IsEmpty())
 		{
-			const FScopedDataLayerTransaction Transaction(LOCTEXT("DataLayerOutlinerChangeDataLayersParent", "Change Data Layers Parent"), GetOwningWorld());
+			const FScopedTransaction Transaction(LOCTEXT("DataLayerOutlinerChangeDataLayersParent", "Change Data Layers Parent"));
 			for (UDataLayerInstance* DataLayerInstance : ValidDataLayers)
 			{
 				DataLayerEditorSubsystem->SetParentDataLayer(DataLayerInstance, ParentDataLayer);
@@ -1018,7 +1018,7 @@ void FDataLayerMode::RegisterContextMenu()
 
 					for (const UDataLayerAsset* DataLayerAsset : DataLayerAssets)
 					{
-						const FScopedDataLayerTransaction Transaction(LOCTEXT("CreateNewDataLayer", "Create New Data Layer"), GetOwningWorld());
+						const FScopedTransaction Transaction(LOCTEXT("CreateNewDataLayer", "Create New Data Layer"));
 						SelectedDataLayersSet.Empty();
 						SelectedDataLayerActors.Empty();
 
@@ -1072,7 +1072,7 @@ void FDataLayerMode::RegisterContextMenu()
 						Section.AddMenuEntry("ImportDataLayers", LOCTEXT("ImportDataLayers", "Import Data Layer(s)"), FText(), FSlateIcon(),
 							FUIAction(FExecuteAction::CreateLambda([this, DataLayerAssetsToImport, CreateNewDataLayer]()
 								{
-									const FScopedDataLayerTransaction Transaction(LOCTEXT("ImportDataLayers", "Import Data Layers"), GetOwningWorld());
+									const FScopedTransaction Transaction(LOCTEXT("ImportDataLayers", "Import Data Layers"));
 									CreateNewDataLayer(nullptr, DataLayerAssetsToImport);
 								})));
 					}
@@ -1096,7 +1096,7 @@ void FDataLayerMode::RegisterContextMenu()
 						FExecuteAction::CreateLambda([this]() {
 							if (UDataLayerAsset* DataLayerAsset = PromptDataLayerAssetSelection())
 							{
-								const FScopedDataLayerTransaction Transaction(LOCTEXT("AddSelectedActorsToNewDataLayer", "Add Selected Actors to New Data Layer"), GetOwningWorld());
+								const FScopedTransaction Transaction(LOCTEXT("AddSelectedActorsToNewDataLayer", "Add Selected Actors to New Data Layer"));
 								FDataLayerCreationParameters CreationParams;
 								CreationParams.DataLayerAsset = DataLayerAsset;
 								CreationParams.WorlDataLayers = GetOwningWorldAWorldDataLayers();
@@ -1113,7 +1113,7 @@ void FDataLayerMode::RegisterContextMenu()
 						FExecuteAction::CreateLambda([this,SelectedDataLayers]() {
 							check(!SelectedDataLayers.IsEmpty());
 							{
-								const FScopedDataLayerTransaction Transaction(LOCTEXT("AddSelectedActorsToSelectedDataLayers", "Add Selected Actors to Selected Data Layers"), GetOwningWorld());
+								const FScopedTransaction Transaction(LOCTEXT("AddSelectedActorsToSelectedDataLayers", "Add Selected Actors to Selected Data Layers"));
 								DataLayerEditorSubsystem->AddSelectedActorsToDataLayers(SelectedDataLayers);
 							}}),
 						FCanExecuteAction::CreateLambda([SelectedDataLayers,bSelectedDataLayersContainsLocked] { return !SelectedDataLayers.IsEmpty() && GEditor->GetSelectedActorCount() > 0 && !bSelectedDataLayersContainsLocked; })
@@ -1131,7 +1131,7 @@ void FDataLayerMode::RegisterContextMenu()
 								Algo::TransformIf(SelectedDataLayerActors, Actors, [](const FSelectedDataLayerActor& InActor) { return InActor.Value.Get(); }, [](const FSelectedDataLayerActor& InActor) { return const_cast<AActor*>(InActor.Value.Get()); });
 								if (!Actors.IsEmpty())
 								{
-									const FScopedDataLayerTransaction Transaction(LOCTEXT("AddSelectedActorsToDataLayer", "Add Selected Actors to Selected Data Layer"), GetOwningWorld());
+									const FScopedTransaction Transaction(LOCTEXT("AddSelectedActorsToDataLayer", "Add Selected Actors to Selected Data Layer"));
 									DataLayerEditorSubsystem->AddActorsToDataLayers(Actors, { TargetDataLayer });
 								}
 							}));
@@ -1164,7 +1164,7 @@ void FDataLayerMode::RegisterContextMenu()
 						FExecuteAction::CreateLambda([this,SelectedDataLayers]() {
 							check(!SelectedDataLayers.IsEmpty());
 							{
-								const FScopedDataLayerTransaction Transaction(LOCTEXT("RemoveSelectedActorsFromSelectedDataLayers_DataLayerMode", "Remove Selected Actors from Selected Data Layers"), GetOwningWorld());
+								const FScopedTransaction Transaction(LOCTEXT("RemoveSelectedActorsFromSelectedDataLayers_DataLayerMode", "Remove Selected Actors from Selected Data Layers"));
 								DataLayerEditorSubsystem->RemoveSelectedActorsFromDataLayers(SelectedDataLayers);
 							}}),
 						FCanExecuteAction::CreateLambda([SelectedDataLayers,bSelectedDataLayersContainsLocked] { return !SelectedDataLayers.IsEmpty() && GEditor->GetSelectedActorCount() > 0 && !bSelectedDataLayersContainsLocked; })
@@ -1175,7 +1175,7 @@ void FDataLayerMode::RegisterContextMenu()
 						FExecuteAction::CreateLambda([this,SelectedDataLayers]() {
 							check(!SelectedDataLayers.IsEmpty());
 							{
-								const FScopedDataLayerTransaction Transaction(LOCTEXT("DeleteSelectedDataLayers", "Delete Selected Data Layers"), GetOwningWorld());
+								const FScopedTransaction Transaction(LOCTEXT("DeleteSelectedDataLayers", "Delete Selected Data Layers"));
 								DataLayerEditorSubsystem->DeleteDataLayers(SelectedDataLayers);
 							}}),
 						FCanExecuteAction::CreateLambda([SelectedDataLayers,bSelectedDataLayersContainsLocked] { return !SelectedDataLayers.IsEmpty() && !bSelectedDataLayersContainsLocked; })
@@ -1207,7 +1207,7 @@ void FDataLayerMode::RegisterContextMenu()
 						FExecuteAction::CreateLambda([this,SelectedDataLayers]() {
 							check(!SelectedDataLayers.IsEmpty());
 							{
-								const FScopedDataLayerTransaction Transaction(LOCTEXT("SelectActorsInDataLayers", "Select Actors in Data Layers"), GetOwningWorld());
+								const FScopedTransaction Transaction(LOCTEXT("SelectActorsInDataLayers", "Select Actors in Data Layers"));
 								GEditor->SelectNone(/*bNoteSelectionChange*/false, /*bDeselectBSPSurfs*/true);
 								DataLayerEditorSubsystem->SelectActorsInDataLayers(SelectedDataLayers, /*bSelect*/true, /*bNotify*/true, /*bSelectEvenIfHidden*/true);
 							}}),
@@ -1219,7 +1219,7 @@ void FDataLayerMode::RegisterContextMenu()
 						FExecuteAction::CreateLambda([this,SelectedDataLayers]() {
 							check(!SelectedDataLayers.IsEmpty());
 							{
-								const FScopedDataLayerTransaction Transaction(LOCTEXT("AppendActorsToSelection", "Append Actors in Data Layer to Selection"), GetOwningWorld());
+								const FScopedTransaction Transaction(LOCTEXT("AppendActorsToSelection", "Append Actors in Data Layer to Selection"));
 								DataLayerEditorSubsystem->SelectActorsInDataLayers(SelectedDataLayers, /*bSelect*/true, /*bNotify*/true, /*bSelectEvenIfHidden*/true);
 							}}),
 						FCanExecuteAction::CreateLambda([SelectedDataLayers,bSelectedDataLayersContainsLocked] { return !SelectedDataLayers.IsEmpty() && !bSelectedDataLayersContainsLocked; })
@@ -1230,7 +1230,7 @@ void FDataLayerMode::RegisterContextMenu()
 						FExecuteAction::CreateLambda([this,SelectedDataLayers]() {
 							check(!SelectedDataLayers.IsEmpty());
 							{
-								const FScopedDataLayerTransaction Transaction(LOCTEXT("DeselectActors", "Deselect Actors in Data Layer"), GetOwningWorld());
+								const FScopedTransaction Transaction(LOCTEXT("DeselectActors", "Deselect Actors in Data Layer"));
 								DataLayerEditorSubsystem->SelectActorsInDataLayers(SelectedDataLayers, /*bSelect*/false, /*bNotifySelectActors*/true);
 							}}),
 						FCanExecuteAction::CreateLambda([SelectedDataLayers] { return !SelectedDataLayers.IsEmpty(); })
@@ -1270,7 +1270,7 @@ void FDataLayerMode::RegisterContextMenu()
 					FUIAction(
 						FExecuteAction::CreateLambda([this, SceneOutliner, SelectedDataLayers]() 
 						{
-							const FScopedDataLayerTransaction Transaction(LOCTEXT("MakeCurrentDataLayers", "Make Current Data Layer(s)"), GetOwningWorld());
+							const FScopedTransaction Transaction(LOCTEXT("MakeCurrentDataLayers", "Make Current Data Layer(s)"));
 							for (TWeakObjectPtr<const UDataLayerInstance>& DataLayer : SelectedDataLayersSet)
 							{
 								if (DataLayer.IsValid() && !DataLayer->IsLocked())
@@ -1286,7 +1286,7 @@ void FDataLayerMode::RegisterContextMenu()
 					FUIAction(
 						FExecuteAction::CreateLambda([this, SceneOutliner, SelectedDataLayers]()
 						{
-							const FScopedDataLayerTransaction Transaction(LOCTEXT("RemoveCurrentDataLayers", "Remove Current Data Layer(s)"), GetOwningWorld());
+							const FScopedTransaction Transaction(LOCTEXT("RemoveCurrentDataLayers", "Remove Current Data Layer(s)"));
 							for (TWeakObjectPtr<const UDataLayerInstance>& DataLayer : SelectedDataLayersSet)
 							{
 								if (DataLayer.IsValid() && !DataLayer->IsLocked())
@@ -1304,7 +1304,7 @@ void FDataLayerMode::RegisterContextMenu()
 						{
 							check(!AllDataLayers.IsEmpty());
 							{
-								const FScopedDataLayerTransaction Transaction(LOCTEXT("ClearCurrentDataLayers", "Clear Current Data Layers"), GetOwningWorld());
+								const FScopedTransaction Transaction(LOCTEXT("ClearCurrentDataLayers", "Clear Current Data Layers"));
 								for (const UDataLayerInstance* DataLayer : AllDataLayers)
 								{
 									DataLayerEditorSubsystem->RemoveFromActorEditorContext(const_cast<UDataLayerInstance*>(DataLayer));
@@ -1322,7 +1322,7 @@ void FDataLayerMode::RegisterContextMenu()
 						FExecuteAction::CreateLambda([this,AllDataLayers]() {
 							check(!AllDataLayers.IsEmpty());
 							{
-								const FScopedDataLayerTransaction Transaction(LOCTEXT("MakeAllDataLayersVisible", "Make All Data Layers Visible"), GetOwningWorld());
+								const FScopedTransaction Transaction(LOCTEXT("MakeAllDataLayersVisible", "Make All Data Layers Visible"));
 								DataLayerEditorSubsystem->MakeAllDataLayersVisible();
 							}}),
 						FCanExecuteAction::CreateLambda([AllDataLayers] { return !AllDataLayers.IsEmpty(); })
@@ -1411,7 +1411,7 @@ void FDataLayerMode::CreateViewContent(FMenuBuilder& MenuBuilder)
 			{
 				if (AWorldDataLayers* WorldDataLayers = RepresentingWorld.IsValid() ? RepresentingWorld->GetWorldDataLayers() : nullptr)
 				{
-					const FScopedDataLayerTransaction Transaction(LOCTEXT("ToggleAllowRuntimeDataLayerEditingTransaction", "Toggle Allow Runtime Data Layer Editing"), GetOwningWorld());
+					const FScopedTransaction Transaction(LOCTEXT("ToggleAllowRuntimeDataLayerEditingTransaction", "Toggle Allow Runtime Data Layer Editing"));
 					WorldDataLayers->SetAllowRuntimeDataLayerEditing(!WorldDataLayers->GetAllowRuntimeDataLayerEditing());
 				}
 				SceneOutliner->FullRefresh();
@@ -1444,7 +1444,7 @@ void FDataLayerMode::CreateViewContent(FMenuBuilder& MenuBuilder)
 		FUIAction(
 			FExecuteAction::CreateLambda([this]()
 			{
-				const FScopedDataLayerTransaction Transaction(LOCTEXT("ResetDataLayerUserSettings", "Reset User Settings"), GetOwningWorld());
+				const FScopedTransaction Transaction(LOCTEXT("ResetDataLayerUserSettings", "Reset User Settings"));
 				DataLayerEditorSubsystem->ResetUserSettings();
 			})
 		)
diff --git a/Engine/Source/Editor/DataLayerEditor/Private/DataLayer/DataLayerNameEditSink.h b/Engine/Source/Editor/DataLayerEditor/Private/DataLayer/DataLayerNameEditSink.h
--- a/Engine/Source/Editor/DataLayerEditor/Private/DataLayer/DataLayerNameEditSink.h
+++ b/Engine/Source/Editor/DataLayerEditor/Private/DataLayer/DataLayerNameEditSink.h
@@ -3,7 +3,7 @@
 #pragma once
 
 #include "DataLayer/DataLayerEditorSubsystem.h"
-#include "DataLayerTransaction.h"
+#include "ScopedTransaction.h"
 #include "IObjectNameEditSink.h"
 
 #define LOCTEXT_NAMESPACE "DataLayer"
@@ -39,7 +39,7 @@ class FDataLayerNameEditSink : public UE::EditorWidgets::IObjectNameEditSink
 		{
 			if (!DisplayName.Equals(DataLayerInstance->GetDataLayerShortName(), ESearchCase::CaseSensitive))
 			{
-				const FScopedDataLayerTransaction Transaction(LOCTEXT("DataLayerNameEditSinkRenameDataLayerTransaction", "Rename Data Layer"), DataLayerInstance->GetWorld());
+				const FScopedTransaction Transaction(LOCTEXT("DataLayerNameEditSinkRenameDataLayerTransaction", "Rename Data Layer"));
 
 				PRAGMA_DISABLE_DEPRECATION_WARNINGS
 				return UDataLayerEditorSubsystem::Get()->RenameDataLayer(DataLayerInstance, *DisplayName);
diff --git a/Engine/Source/Editor/DataLayerEditor/Private/DataLayer/DataLayerOutlinerIsLoadedInEditorColumn.cpp b/Engine/Source/Editor/DataLayerEditor/Private/DataLayer/DataLayerOutlinerIsLoadedInEditorColumn.cpp
--- a/Engine/Source/Editor/DataLayerEditor/Private/DataLayer/DataLayerOutlinerIsLoadedInEditorColumn.cpp
+++ b/Engine/Source/Editor/DataLayerEditor/Private/DataLayer/DataLayerOutlinerIsLoadedInEditorColumn.cpp
@@ -7,7 +7,7 @@
 #include "Widgets/Input/SButton.h"
 #include "Widgets/Input/SCheckBox.h"
 #include "Widgets/Views/STreeView.h"
-#include "DataLayerTransaction.h"
+#include "ScopedTransaction.h"
 #include "Styling/AppStyle.h"
 #include "Editor.h"
 
@@ -91,12 +91,12 @@ const TSharedRef<SWidget> FDataLayerOutlinerIsLoadedInEditorColumn::ConstructRow
 								}
 							}
 
-							const FScopedDataLayerTransaction Transaction(LOCTEXT("ToggleDataLayersIsLoadedInEditor", "Toggle Data Layers Dynamically Loaded In Editor Flag"), World);
+							const FScopedTransaction Transaction(LOCTEXT("ToggleDataLayersIsLoadedInEditor", "Toggle Data Layers Dynamically Loaded In Editor Flag"));
 							DataLayerEditorSubsystem->ToggleDataLayersIsLoadedInEditor(AllSelectedDataLayers, /*bIsFromUserChange*/true);
 						}
 						else
 						{
-							const FScopedDataLayerTransaction Transaction(LOCTEXT("ToggleDataLayerIsLoadedInEditor", "Toggle Data Layer Dynamically Loaded In Editor Flag"), World);
+							const FScopedTransaction Transaction(LOCTEXT("ToggleDataLayerIsLoadedInEditor", "Toggle Data Layer Dynamically Loaded In Editor Flag"));
 							DataLayerEditorSubsystem->ToggleDataLayerIsLoadedInEditor(DataLayer, /*bIsFromUserChange*/true);
 						}
 					}
diff --git a/Engine/Source/Editor/DataLayerEditor/Private/DataLayer/DataLayerTransaction.h b/Engine/Source/Editor/DataLayerEditor/Private/DataLayer/DataLayerTransaction.h
deleted file mode 100644
diff --git a/Engine/Source/Editor/DataLayerEditor/Private/DataLayer/SDataLayerOutliner.cpp b/Engine/Source/Editor/DataLayerEditor/Private/DataLayer/SDataLayerOutliner.cpp
--- a/Engine/Source/Editor/DataLayerEditor/Private/DataLayer/SDataLayerOutliner.cpp
+++ b/Engine/Source/Editor/DataLayerEditor/Private/DataLayer/SDataLayerOutliner.cpp
@@ -8,7 +8,7 @@
 #include "DataLayer/DataLayerTreeItem.h"
 #include "DataLayer/DataLayerEditorSubsystem.h"
 #include "WorldPartition/DataLayer/DataLayerInstance.h"
-#include "DataLayerTransaction.h"
+#include "ScopedTransaction.h"
 #include "Algo/Transform.h"
 #include "Editor.h"
 
@@ -81,7 +81,7 @@ FReply SDataLayerOutliner::OnAddSelectedActorsToSelectedDataLayersClicked()
 	if (CanAddSelectedActorsToSelectedDataLayersClicked())
 	{
 		TArray<UDataLayerInstance*> SelectedDataLayers = GetSelectedDataLayers();
-		const FScopedDataLayerTransaction Transaction(LOCTEXT("AddSelectedActorsToSelectedDataLayers", "Add Selected Actors to Selected Data Layers"), SelectedDataLayers[0]->GetWorld());
+		const FScopedTransaction Transaction(LOCTEXT("AddSelectedActorsToSelectedDataLayers", "Add Selected Actors to Selected Data Layers"));
 		UDataLayerEditorSubsystem::Get()->AddSelectedActorsToDataLayers(SelectedDataLayers);
 	}
 	return FReply::Handled();
@@ -92,7 +92,7 @@ FReply SDataLayerOutliner::OnRemoveSelectedActorsFromSelectedDataLayersClicked()
 	if (CanRemoveSelectedActorsFromSelectedDataLayersClicked())
 	{
 		TArray<UDataLayerInstance*> SelectedDataLayers = GetSelectedDataLayers();
-		const FScopedDataLayerTransaction Transaction(LOCTEXT("RemoveSelectedActorsFromSelectedDataLayers", "Remove Selected Actors from Selected Data Layers"), SelectedDataLayers[0]->GetWorld());
+		const FScopedTransaction Transaction(LOCTEXT("RemoveSelectedActorsFromSelectedDataLayers", "Remove Selected Actors from Selected Data Layers"));
 		UDataLayerEditorSubsystem::Get()->RemoveSelectedActorsFromDataLayers(SelectedDataLayers);
 	}
 	return FReply::Handled();
diff --git a/Engine/Source/Editor/DataLayerEditor/Private/DataLayer/SDataLayerTreeLabel.cpp b/Engine/Source/Editor/DataLayerEditor/Private/DataLayer/SDataLayerTreeLabel.cpp
--- a/Engine/Source/Editor/DataLayerEditor/Private/DataLayer/SDataLayerTreeLabel.cpp
+++ b/Engine/Source/Editor/DataLayerEditor/Private/DataLayer/SDataLayerTreeLabel.cpp
@@ -8,7 +8,7 @@
 #include "Math/ColorList.h"
 #include "ISceneOutliner.h"
 #include "ISceneOutlinerMode.h"
-#include "DataLayerTransaction.h"
+#include "ScopedTransaction.h"
 #include "Editor.h"
 
 #define LOCTEXT_NAMESPACE "DataLayer"
@@ -278,7 +278,7 @@ void SDataLayerTreeLabel::OnLabelCommitted(const FText& InLabel, ETextCommit::Ty
 	check(DataLayerInstance->SupportRelabeling());
 	if (!InLabel.ToString().Equals(DataLayerInstance->GetDataLayerShortName(), ESearchCase::CaseSensitive))
 	{
-		const FScopedDataLayerTransaction Transaction(LOCTEXT("SceneOutlinerRenameDataLayerTransaction", "Rename Data Layer"), DataLayerInstance->GetWorld());
+		const FScopedTransaction Transaction(LOCTEXT("SceneOutlinerRenameDataLayerTransaction", "Rename Data Layer"));
 
 		PRAGMA_DISABLE_DEPRECATION_WARNINGS
 		UDataLayerEditorSubsystem::Get()->RenameDataLayer(DataLayerInstance, *InLabel.ToString());
