Fixed rare crash in UDataLayerEditorSubsystem when loading another world partition map.
#rnx

diff --git a/Engine/Source/Editor/DataLayerEditor/Private/DataLayer/DataLayerEditorSubsystem.cpp b/Engine/Source/Editor/DataLayerEditor/Private/DataLayer/DataLayerEditorSubsystem.cpp
--- a/Engine/Source/Editor/DataLayerEditor/Private/DataLayer/DataLayerEditorSubsystem.cpp
+++ b/Engine/Source/Editor/DataLayerEditor/Private/DataLayer/DataLayerEditorSubsystem.cpp
@@ -177,42 +177,45 @@ void UDataLayerEditorSubsystem::Deinitialize()
 
 void UDataLayerEditorSubsystem::OnExecuteActorEditorContextAction(UWorld* InWorld, const EActorEditorContextAction& InType, AActor* InActor)
 {	
-	check(InWorld == GetWorld());
-	UDataLayerSubsystem* DataLayerSubsystem = GetWorld()->GetSubsystem<UDataLayerSubsystem>();
-
-	switch (InType)
+	UE_CLOG(!InWorld, LogDataLayerEditorSubsystem, Error, TEXT("%s - Failed because world in null."), ANSI_TO_TCHAR(__FUNCTION__));
+	if (UDataLayerSubsystem* DataLayerSubsystem = UWorld::GetSubsystem<UDataLayerSubsystem>(InWorld))
 	{
-	case EActorEditorContextAction::ApplyContext:
-		check(InActor && InActor->GetWorld() == InWorld);
-		{
-			AddActorToDataLayers(InActor, DataLayerSubsystem->GetActorEditorContextDataLayers());
-		}
-		break;
-	case EActorEditorContextAction::ResetContext:
-		for (UDataLayerInstance* DataLayer : DataLayerSubsystem->GetActorEditorContextDataLayers())
+		switch (InType)
 		{
-			RemoveFromActorEditorContext(DataLayer);
+		case EActorEditorContextAction::ApplyContext:
+			check(InActor && InActor->GetWorld() == InWorld);
+			{
+				AddActorToDataLayers(InActor, DataLayerSubsystem->GetActorEditorContextDataLayers());
+			}
+			break;
+		case EActorEditorContextAction::ResetContext:
+			for (UDataLayerInstance* DataLayer : DataLayerSubsystem->GetActorEditorContextDataLayers())
+			{
+				RemoveFromActorEditorContext(DataLayer);
+			}
+			break;
+		case EActorEditorContextAction::PushContext:
+			DataLayerSubsystem->PushActorEditorContext();
+			BroadcastDataLayerChanged(EDataLayerAction::Reset, NULL, NAME_None);
+			break;
+		case EActorEditorContextAction::PopContext:
+			DataLayerSubsystem->PopActorEditorContext();
+			BroadcastDataLayerChanged(EDataLayerAction::Reset, NULL, NAME_None);
+			break;
 		}
-		break;
-	case EActorEditorContextAction::PushContext:
-		DataLayerSubsystem->PushActorEditorContext();
-		BroadcastDataLayerChanged(EDataLayerAction::Reset, NULL, NAME_None);
-		break;
-	case EActorEditorContextAction::PopContext:
-		DataLayerSubsystem->PopActorEditorContext();
-		BroadcastDataLayerChanged(EDataLayerAction::Reset, NULL, NAME_None);
-		break;
 	}
 }	
 
 bool UDataLayerEditorSubsystem::GetActorEditorContextDisplayInfo(UWorld* InWorld, FActorEditorContextClientDisplayInfo& OutDiplayInfo) const
 {
-	UDataLayerSubsystem* DataLayerSubsystem = GetWorld()->GetSubsystem<UDataLayerSubsystem>();
-	if (!DataLayerSubsystem->GetActorEditorContextDataLayers().IsEmpty())
+	if (UDataLayerSubsystem* DataLayerSubsystem = UWorld::GetSubsystem<UDataLayerSubsystem>(InWorld))
 	{
-		OutDiplayInfo.Title = TEXT("Data Layers");
-		OutDiplayInfo.Brush = FAppStyle::GetBrush(TEXT("DataLayer.Editor"));
-		return true;
+		if (!DataLayerSubsystem->GetActorEditorContextDataLayers().IsEmpty())
+		{
+			OutDiplayInfo.Title = TEXT("Data Layers");
+			OutDiplayInfo.Brush = FAppStyle::GetBrush(TEXT("DataLayer.Editor"));
+			return true;
+		}
 	}
 	return false;
 }
@@ -220,33 +223,35 @@ bool UDataLayerEditorSubsystem::GetActorEditorContextDisplayInfo(UWorld* InWorld
 TSharedRef<SWidget> UDataLayerEditorSubsystem::GetActorEditorContextWidget(UWorld* InWorld) const
 {
 	TSharedRef<SVerticalBox> OutWidget = SNew(SVerticalBox);
-	
-	UDataLayerSubsystem* DataLayerSubsystem = GetWorld()->GetSubsystem<UDataLayerSubsystem>();
-	TArray<UDataLayerInstance*> DataLayers = DataLayerSubsystem->GetActorEditorContextDataLayers();
-	for (UDataLayerInstance* DataLayer : DataLayers)
+
+	if (UDataLayerSubsystem* DataLayerSubsystem = UWorld::GetSubsystem<UDataLayerSubsystem>(InWorld))
 	{
-		check(IsValid(DataLayer));
-		OutWidget->AddSlot().AutoHeight()
-		[
-			SNew(SHorizontalBox)
-			+ SHorizontalBox::Slot()
-			.AutoWidth()
-			.Padding(0.0f, 1.0f, 1.0f, 1.0f)
-			.VAlign(VAlign_Center)
-			[
-				SNew(SImage)
-				.ColorAndOpacity(DataLayer->GetDebugColor())
-				.Image(FAppStyle::Get().GetBrush("DataLayer.ColorIcon"))
-				.DesiredSizeOverride(FVector2D(8, 8))
-			]
-			+ SHorizontalBox::Slot()
-			.AutoWidth()
-			.Padding(4.0f, 1.0f, 1.0f, 1.0f)
+		TArray<UDataLayerInstance*> DataLayers = DataLayerSubsystem->GetActorEditorContextDataLayers();
+		for (UDataLayerInstance* DataLayer : DataLayers)
+		{
+			check(IsValid(DataLayer));
+			OutWidget->AddSlot().AutoHeight()
 			[
-				SNew(STextBlock)
-				.Text(FText::FromString(DataLayer->GetDataLayerShortName()))
-			]
-		];
+				SNew(SHorizontalBox)
+				+ SHorizontalBox::Slot()
+				.AutoWidth()
+				.Padding(0.0f, 1.0f, 1.0f, 1.0f)
+				.VAlign(VAlign_Center)
+				[
+					SNew(SImage)
+					.ColorAndOpacity(DataLayer->GetDebugColor())
+					.Image(FAppStyle::Get().GetBrush("DataLayer.ColorIcon"))
+					.DesiredSizeOverride(FVector2D(8, 8))
+				]
+				+ SHorizontalBox::Slot()
+				.AutoWidth()
+				.Padding(4.0f, 1.0f, 1.0f, 1.0f)
+				[
+					SNew(STextBlock)
+					.Text(FText::FromString(DataLayer->GetDataLayerShortName()))
+				]
+			];
+		}
 	}
 	
 	return OutWidget;
@@ -881,19 +886,22 @@ void UDataLayerEditorSubsystem::ToggleDataLayersVisibility(const TArray<UDataLay
 
 void UDataLayerEditorSubsystem::MakeAllDataLayersVisible()
 {
-	UDataLayerSubsystem* DataLayerSubsystem = UWorld::GetSubsystem<UDataLayerSubsystem>(GetWorld());
-	DataLayerSubsystem->ForEachDataLayer([this](UDataLayerInstance* DataLayer)
+	UE_CLOG(!GetWorld(), LogDataLayerEditorSubsystem, Error, TEXT("%s - Failed because world in null."), ANSI_TO_TCHAR(__FUNCTION__));
+	if (UDataLayerSubsystem* DataLayerSubsystem = UWorld::GetSubsystem<UDataLayerSubsystem>(GetWorld()))
 	{
-		if (!DataLayer->IsVisible())
+		DataLayerSubsystem->ForEachDataLayer([this](UDataLayerInstance* DataLayer)
 		{
-			DataLayer->Modify();
-			DataLayer->SetVisible(true);
-			BroadcastDataLayerChanged(EDataLayerAction::Modify, DataLayer, "bIsVisible");
-		}
-		return true;
-	});
-	
-	UpdateAllActorsVisibility(true, true);
+			if (!DataLayer->IsVisible())
+			{
+				DataLayer->Modify();
+				DataLayer->SetVisible(true);
+				BroadcastDataLayerChanged(EDataLayerAction::Modify, DataLayer, "bIsVisible");
+			}
+			return true;
+		});
+
+		UpdateAllActorsVisibility(true, true);
+	}
 }
 
 bool UDataLayerEditorSubsystem::SetDataLayerIsLoadedInEditorInternal(UDataLayerInstance* DataLayer, const bool bIsLoadedInEditor, const bool bIsFromUserChange)
@@ -966,72 +974,86 @@ bool UDataLayerEditorSubsystem::ToggleDataLayersIsLoadedInEditor(const TArray<UD
 bool UDataLayerEditorSubsystem::ResetUserSettings()
 {
 	bool bChanged = false;
-	UDataLayerSubsystem* DataLayerSubsystem = UWorld::GetSubsystem<UDataLayerSubsystem>(GetWorld());
-	DataLayerSubsystem->ForEachDataLayer([this, &bChanged](UDataLayerInstance* DataLayer)
+	UE_CLOG(!GetWorld(), LogDataLayerEditorSubsystem, Error, TEXT("%s - Failed because world in null."), ANSI_TO_TCHAR(__FUNCTION__));
+	if (UDataLayerSubsystem* DataLayerSubsystem = UWorld::GetSubsystem<UDataLayerSubsystem>(GetWorld()))
 	{
-		bChanged |= SetDataLayerIsLoadedInEditorInternal(DataLayer, DataLayer->IsInitiallyLoadedInEditor(), true);
-		return true;
-	});
+		DataLayerSubsystem->ForEachDataLayer([this, &bChanged](UDataLayerInstance* DataLayer)
+		{
+			bChanged |= SetDataLayerIsLoadedInEditorInternal(DataLayer, DataLayer->IsInitiallyLoadedInEditor(), true);
+			return true;
+		});
 	
-	if (bChanged)
-	{
-		OnDataLayerEditorLoadingStateChanged(true);
+		if (bChanged)
+		{
+			OnDataLayerEditorLoadingStateChanged(true);
+		}
 	}
-
 	return true;
 }
 
 bool UDataLayerEditorSubsystem::HasDeprecatedDataLayers() const
 {
-	if (AWorldDataLayers* WorldDataLayers = GetWorld()->GetWorldDataLayers())
+	UWorld* World = GetWorld();
+	if (AWorldDataLayers* WorldDataLayers = World ? World->GetWorldDataLayers() : nullptr)
 	{
 		return WorldDataLayers->HasDeprecatedDataLayers();
 	}
-
 	return false;
 }
 
 UDataLayerInstance* UDataLayerEditorSubsystem::GetDataLayerInstance(const FName& DataLayerInstanceName) const
 {
-	UDataLayerSubsystem* DataLayerSubsystem = UWorld::GetSubsystem<UDataLayerSubsystem>(GetWorld());
-	return DataLayerSubsystem->GetDataLayerInstance(DataLayerInstanceName);
+	if (UDataLayerSubsystem* DataLayerSubsystem = UWorld::GetSubsystem<UDataLayerSubsystem>(GetWorld()))
+	{
+		return DataLayerSubsystem->GetDataLayerInstance(DataLayerInstanceName);
+	}
+	return nullptr;
 }
 
 UDataLayerInstance* UDataLayerEditorSubsystem::GetDataLayerInstance(const UDataLayerAsset* DataLayerAsset) const
 {
-	UDataLayerSubsystem* DataLayerSubsystem = UWorld::GetSubsystem<UDataLayerSubsystem>(GetWorld());
-	return DataLayerSubsystem->GetDataLayerInstance(DataLayerAsset);
+	if (UDataLayerSubsystem* DataLayerSubsystem = UWorld::GetSubsystem<UDataLayerSubsystem>(GetWorld()))
+	{
+		return DataLayerSubsystem->GetDataLayerInstance(DataLayerAsset);
+	}
+	return nullptr;
 }
 
 void UDataLayerEditorSubsystem::AddAllDataLayersTo(TArray<TWeakObjectPtr<UDataLayerInstance>>& OutDataLayers) const
 {
-	UDataLayerSubsystem* DataLayerSubsystem = UWorld::GetSubsystem<UDataLayerSubsystem>(GetWorld());
-	DataLayerSubsystem->ForEachDataLayer([&OutDataLayers](UDataLayerInstance* DataLayer)
+	UE_CLOG(!GetWorld(), LogDataLayerEditorSubsystem, Error, TEXT("%s - Failed because world in null."), ANSI_TO_TCHAR(__FUNCTION__));
+	if (UDataLayerSubsystem* DataLayerSubsystem = UWorld::GetSubsystem<UDataLayerSubsystem>(GetWorld()))
 	{
-		OutDataLayers.Add(DataLayer);
-		return true;
-	});
+		DataLayerSubsystem->ForEachDataLayer([&OutDataLayers](UDataLayerInstance* DataLayer)
+		{
+			OutDataLayers.Add(DataLayer);
+			return true;
+		});
+	}
 }
 
 UDataLayerInstance* UDataLayerEditorSubsystem::CreateDataLayerInstance(const FDataLayerCreationParameters& Parameters)
 {
 	UDataLayerInstance* NewDataLayer = nullptr;
 
-	AWorldDataLayers* WorldDataLayers = Parameters.WorlDataLayers != nullptr ? Parameters.WorlDataLayers.Get() : GetWorld()->GetWorldDataLayers();
-	check(!Parameters.ParentDataLayer || !Parameters.WorlDataLayers || (Parameters.ParentDataLayer->GetOuterAWorldDataLayers() == Parameters.WorlDataLayers));
-	if (!WorldDataLayers->HasDeprecatedDataLayers())
-	{
-		NewDataLayer = WorldDataLayers->CreateDataLayer<UDataLayerInstanceWithAsset>(Parameters.DataLayerAsset);
-	}
-	else
+	UE_CLOG(!GetWorld(), LogDataLayerEditorSubsystem, Error, TEXT("%s - Failed because world in null."), ANSI_TO_TCHAR(__FUNCTION__));
+	if (AWorldDataLayers* WorldDataLayers = Parameters.WorlDataLayers != nullptr ? Parameters.WorlDataLayers.Get() : (GetWorld() ? GetWorld()->GetWorldDataLayers() : nullptr))
 	{
-		NewDataLayer = WorldDataLayers->CreateDataLayer<UDeprecatedDataLayerInstance>();
-	}
+		check(!Parameters.ParentDataLayer || !Parameters.WorlDataLayers || (Parameters.ParentDataLayer->GetOuterAWorldDataLayers() == Parameters.WorlDataLayers));
+		if (!WorldDataLayers->HasDeprecatedDataLayers())
+		{
+			NewDataLayer = WorldDataLayers->CreateDataLayer<UDataLayerInstanceWithAsset>(Parameters.DataLayerAsset);
+		}
+		else
+		{
+			NewDataLayer = WorldDataLayers->CreateDataLayer<UDeprecatedDataLayerInstance>();
+		}
 
-	if (NewDataLayer != nullptr)
-	{
-		BroadcastDataLayerChanged(EDataLayerAction::Add, NewDataLayer, NAME_None);
-		SetParentDataLayer(NewDataLayer, Parameters.ParentDataLayer);
+		if (NewDataLayer != nullptr)
+		{
+			BroadcastDataLayerChanged(EDataLayerAction::Add, NewDataLayer, NAME_None);
+			SetParentDataLayer(NewDataLayer, Parameters.ParentDataLayer);
+		}
 	}
 	
 	return NewDataLayer;
@@ -1039,19 +1061,25 @@ UDataLayerInstance* UDataLayerEditorSubsystem::CreateDataLayerInstance(const FDa
 
 void UDataLayerEditorSubsystem::DeleteDataLayers(const TArray<UDataLayerInstance*>& DataLayersToDelete)
 {
-	UDataLayerSubsystem* DataLayerSubsystem = UWorld::GetSubsystem<UDataLayerSubsystem>(GetWorld());
-	if (DataLayerSubsystem->RemoveDataLayers(DataLayersToDelete))
+	UE_CLOG(!GetWorld(), LogDataLayerEditorSubsystem, Error, TEXT("%s - Failed because world in null."), ANSI_TO_TCHAR(__FUNCTION__));
+	if (UDataLayerSubsystem* DataLayerSubsystem = UWorld::GetSubsystem<UDataLayerSubsystem>(GetWorld()))
 	{
-		BroadcastDataLayerChanged(EDataLayerAction::Delete, NULL, NAME_None);
+		if (DataLayerSubsystem->RemoveDataLayers(DataLayersToDelete))
+		{
+			BroadcastDataLayerChanged(EDataLayerAction::Delete, NULL, NAME_None);
+		}
 	}
 }
 
 void UDataLayerEditorSubsystem::DeleteDataLayer(UDataLayerInstance* DataLayerToDelete)
 {
-	UDataLayerSubsystem* DataLayerSubsystem = UWorld::GetSubsystem<UDataLayerSubsystem>(GetWorld());
-	if (DataLayerSubsystem->RemoveDataLayer(DataLayerToDelete))
+	UE_CLOG(!GetWorld(), LogDataLayerEditorSubsystem, Error, TEXT("%s - Failed because world in null."), ANSI_TO_TCHAR(__FUNCTION__));
+	if (UDataLayerSubsystem* DataLayerSubsystem = UWorld::GetSubsystem<UDataLayerSubsystem>(GetWorld()))
 	{
-		BroadcastDataLayerChanged(EDataLayerAction::Delete, NULL, NAME_None);
+		if (DataLayerSubsystem->RemoveDataLayer(DataLayerToDelete))
+		{
+			BroadcastDataLayerChanged(EDataLayerAction::Delete, NULL, NAME_None);
+		}
 	}
 }
 
@@ -1078,10 +1106,12 @@ void UDataLayerEditorSubsystem::OnDataLayerEditorLoadingStateChanged(bool bIsFro
 
 void UDataLayerEditorSubsystem::BroadcastDataLayerEditorLoadingStateChanged(bool bIsFromUserChange)
 {
+	UE_CLOG(!GetWorld(), LogDataLayerEditorSubsystem, Error, TEXT("%s - Failed because world in null."), ANSI_TO_TCHAR(__FUNCTION__));
 	DataLayerEditorLoadingStateChanged.Broadcast(bIsFromUserChange);
-	
-	UDataLayerSubsystem* DataLayerSubsystem = UWorld::GetSubsystem<UDataLayerSubsystem>(GetWorld());
-	DataLayerSubsystem->UpdateDataLayerEditorPerProjectUserSettings();
+	if (UDataLayerSubsystem* DataLayerSubsystem = UWorld::GetSubsystem<UDataLayerSubsystem>(GetWorld()))
+	{
+		DataLayerSubsystem->UpdateDataLayerEditorPerProjectUserSettings();
+	}
 }
 
 void UDataLayerEditorSubsystem::OnSelectionChanged()
@@ -1135,8 +1165,11 @@ bool UDataLayerEditorSubsystem::TryGetDataLayerFromLabel(const FName& DataLayerL
 
 UDataLayerInstance* UDataLayerEditorSubsystem::GetDataLayerFromLabel(const FName& DataLayerLabel) const
 {
-	UDataLayerSubsystem* DataLayerSubsystem = UWorld::GetSubsystem<UDataLayerSubsystem>(GetWorld());
-	return DataLayerSubsystem->GetDataLayerFromLabel(DataLayerLabel);
+	if (UDataLayerSubsystem* DataLayerSubsystem = UWorld::GetSubsystem<UDataLayerSubsystem>(GetWorld()))
+	{
+		return DataLayerSubsystem->GetDataLayerFromLabel(DataLayerLabel);
+	}
+	return nullptr;
 }
 
 PRAGMA_ENABLE_DEPRECATION_WARNINGS
