UDataLayerEditorSubsystem now passes the DataLayerInstance when broadcasting the EDataLayerAction::Delete action.

diff --git a/Engine/Source/Editor/DataLayerEditor/Private/DataLayer/DataLayerEditorSubsystem.cpp b/Engine/Source/Editor/DataLayerEditor/Private/DataLayer/DataLayerEditorSubsystem.cpp
--- a/Engine/Source/Editor/DataLayerEditor/Private/DataLayer/DataLayerEditorSubsystem.cpp
+++ b/Engine/Source/Editor/DataLayerEditor/Private/DataLayer/DataLayerEditorSubsystem.cpp
@@ -1066,7 +1066,10 @@ void UDataLayerEditorSubsystem::DeleteDataLayers(const TArray<UDataLayerInstance
 	{
 		if (DataLayerSubsystem->RemoveDataLayers(DataLayersToDelete))
 		{
-			BroadcastDataLayerChanged(EDataLayerAction::Delete, NULL, NAME_None);
+			for (UDataLayerInstance* DataLayerInstance : DataLayersToDelete)
+			{
+				BroadcastDataLayerChanged(EDataLayerAction::Delete, DataLayerInstance, NAME_None);
+			}
 		}
 	}
 }
@@ -1078,7 +1081,7 @@ void UDataLayerEditorSubsystem::DeleteDataLayer(UDataLayerInstance* DataLayerToD
 	{
 		if (DataLayerSubsystem->RemoveDataLayer(DataLayerToDelete))
 		{
-			BroadcastDataLayerChanged(EDataLayerAction::Delete, NULL, NAME_None);
+			BroadcastDataLayerChanged(EDataLayerAction::Delete, DataLayerToDelete, NAME_None);
 		}
 	}
 }
