Removed AActor::IsValidForDataLayer in favor of always using AActor::SupportsDataLayer().

diff --git a/Engine/Source/Editor/DataLayerEditor/Private/DataLayer/DataLayerEditorSubsystem.cpp b/Engine/Source/Editor/DataLayerEditor/Private/DataLayer/DataLayerEditorSubsystem.cpp
--- a/Engine/Source/Editor/DataLayerEditor/Private/DataLayer/DataLayerEditorSubsystem.cpp
+++ b/Engine/Source/Editor/DataLayerEditor/Private/DataLayer/DataLayerEditorSubsystem.cpp
@@ -176,34 +176,34 @@ void UDataLayerEditorSubsystem::Deinitialize()
 }
 
 void UDataLayerEditorSubsystem::OnExecuteActorEditorContextAction(UWorld* InWorld, const EActorEditorContextAction& InType, AActor* InActor)
-{	
+{
 	UE_CLOG(!InWorld, LogDataLayerEditorSubsystem, Error, TEXT("%s - Failed because world in null."), ANSI_TO_TCHAR(__FUNCTION__));
 	if (UDataLayerSubsystem* DataLayerSubsystem = UWorld::GetSubsystem<UDataLayerSubsystem>(InWorld))
 	{
-		switch (InType)
+	switch (InType)
+	{
+	case EActorEditorContextAction::ApplyContext:
+		check(InActor && InActor->GetWorld() == InWorld);
 		{
-		case EActorEditorContextAction::ApplyContext:
-			check(InActor && InActor->GetWorld() == InWorld);
-			{
 				AddActorToDataLayers(InActor, DataLayerSubsystem->GetActorEditorContextDataLayers());
-			}
-			break;
-		case EActorEditorContextAction::ResetContext:
+		}
+		break;
+	case EActorEditorContextAction::ResetContext:
 			for (UDataLayerInstance* DataLayer : DataLayerSubsystem->GetActorEditorContextDataLayers())
-			{
-				RemoveFromActorEditorContext(DataLayer);
-			}
-			break;
-		case EActorEditorContextAction::PushContext:
+		{
+			RemoveFromActorEditorContext(DataLayer);
+		}
+		break;
+	case EActorEditorContextAction::PushContext:
 			DataLayerSubsystem->PushActorEditorContext();
-			BroadcastDataLayerChanged(EDataLayerAction::Reset, NULL, NAME_None);
-			break;
-		case EActorEditorContextAction::PopContext:
+		BroadcastDataLayerChanged(EDataLayerAction::Reset, NULL, NAME_None);
+		break;
+	case EActorEditorContextAction::PopContext:
 			DataLayerSubsystem->PopActorEditorContext();
-			BroadcastDataLayerChanged(EDataLayerAction::Reset, NULL, NAME_None);
-			break;
-		}
+		BroadcastDataLayerChanged(EDataLayerAction::Reset, NULL, NAME_None);
+		break;
 	}
+}
 }	
 
 bool UDataLayerEditorSubsystem::GetActorEditorContextDisplayInfo(UWorld* InWorld, FActorEditorContextClientDisplayInfo& OutDiplayInfo) const
@@ -211,11 +211,11 @@ bool UDataLayerEditorSubsystem::GetActorEditorContextDisplayInfo(UWorld* InWorld
 	if (UDataLayerSubsystem* DataLayerSubsystem = UWorld::GetSubsystem<UDataLayerSubsystem>(InWorld))
 	{
 		if (!DataLayerSubsystem->GetActorEditorContextDataLayers().IsEmpty())
-		{
+	{
 			OutDiplayInfo.Title = TEXT("Data Layers");
 			OutDiplayInfo.Brush = FAppStyle::GetBrush(TEXT("DataLayer.Editor"));
-			return true;
-		}
+		return true;
+	}
 	}
 	return false;
 }
@@ -305,7 +305,8 @@ void UDataLayerEditorSubsystem::PostUndoRedo()
 
 bool UDataLayerEditorSubsystem::IsActorValidForDataLayer(AActor* Actor)
 {
-	return Actor && Actor->SupportsDataLayer() && Actor->IsValidForDataLayer() && ((Actor->GetLevel() == Actor->GetWorld()->PersistentLevel) || Actor->GetLevel()->GetWorldDataLayers());
+	UWorld* World = Actor ? Actor->GetWorld() : nullptr;
+	return World && (World->WorldType == EWorldType::Editor) && World->IsPartitionedWorld() && Actor && Actor->SupportsDataLayer() && ((Actor->GetLevel() == Actor->GetWorld()->PersistentLevel) || Actor->GetLevel()->GetWorldDataLayers());
 }
 
 void UDataLayerEditorSubsystem::OnWorldPartitionInitialized(UWorldPartition* InWorldPartition)
@@ -681,18 +682,18 @@ bool UDataLayerEditorSubsystem::UpdateActorVisibility(AActor* Actor, bool& bOutS
 
 	bool bActorBelongsToVisibleDataLayer = false;
 	for (const UDataLayerInstance* DataLayerInstance : DataLayerInstances)
-	{
-		if (DataLayerInstance->IsEffectiveVisible())
 		{
-			if (Actor->SetIsHiddenEdLayer(false))
+		if (DataLayerInstance->IsEffectiveVisible())
 			{
-				bOutActorModified = true;
-			}
-			// Stop, because we found at least one visible DataLayer the actor belongs to
-			bActorBelongsToVisibleDataLayer = true;
+				if (Actor->SetIsHiddenEdLayer(false))
+				{
+					bOutActorModified = true;
+				}
+				// Stop, because we found at least one visible DataLayer the actor belongs to
+				bActorBelongsToVisibleDataLayer = true;
 			break;
-		}
-		}
+			}
+	}
 
 	// If the actor isn't part of a visible DataLayer, hide and de-select it.
 	if (!bActorBelongsToVisibleDataLayer)
@@ -899,9 +900,9 @@ void UDataLayerEditorSubsystem::MakeAllDataLayersVisible()
 			}
 			return true;
 		});
-
-		UpdateAllActorsVisibility(true, true);
-	}
+	
+	UpdateAllActorsVisibility(true, true);
+}
 }
 
 bool UDataLayerEditorSubsystem::SetDataLayerIsLoadedInEditorInternal(UDataLayerInstance* DataLayer, const bool bIsLoadedInEditor, const bool bIsFromUserChange)
@@ -1004,9 +1005,9 @@ bool UDataLayerEditorSubsystem::HasDeprecatedDataLayers() const
 UDataLayerInstance* UDataLayerEditorSubsystem::GetDataLayerInstance(const FName& DataLayerInstanceName) const
 {
 	if (UDataLayerSubsystem* DataLayerSubsystem = UWorld::GetSubsystem<UDataLayerSubsystem>(GetWorld()))
-	{
+{
 		return DataLayerSubsystem->GetDataLayerInstance(DataLayerInstanceName);
-	}
+}
 	return nullptr;
 }
 
@@ -1025,10 +1026,10 @@ void UDataLayerEditorSubsystem::AddAllDataLayersTo(TArray<TWeakObjectPtr<UDataLa
 	if (UDataLayerSubsystem* DataLayerSubsystem = UWorld::GetSubsystem<UDataLayerSubsystem>(GetWorld()))
 	{
 		DataLayerSubsystem->ForEachDataLayer([&OutDataLayers](UDataLayerInstance* DataLayer)
-		{
-			OutDataLayers.Add(DataLayer);
-			return true;
-		});
+			{
+				OutDataLayers.Add(DataLayer);
+				return true;
+			});
 	}
 }
 
@@ -1040,18 +1041,18 @@ UDataLayerInstance* UDataLayerEditorSubsystem::CreateDataLayerInstance(const FDa
 	if (AWorldDataLayers* WorldDataLayers = Parameters.WorlDataLayers != nullptr ? Parameters.WorlDataLayers.Get() : (GetWorld() ? GetWorld()->GetWorldDataLayers() : nullptr))
 	{
 		check(!Parameters.ParentDataLayer || !Parameters.WorlDataLayers || (Parameters.ParentDataLayer->GetOuterAWorldDataLayers() == Parameters.WorlDataLayers));
-		if (!WorldDataLayers->HasDeprecatedDataLayers())
-		{
+	if (!WorldDataLayers->HasDeprecatedDataLayers())
+	{
 			NewDataLayer = WorldDataLayers->CreateDataLayer<UDataLayerInstanceWithAsset>(Parameters.DataLayerAsset);
-		}
-		else
-		{
-			NewDataLayer = WorldDataLayers->CreateDataLayer<UDeprecatedDataLayerInstance>();
-		}
+	}
+	else
+	{
+		NewDataLayer = WorldDataLayers->CreateDataLayer<UDeprecatedDataLayerInstance>();
+	}
 
-		if (NewDataLayer != nullptr)
-		{
-			BroadcastDataLayerChanged(EDataLayerAction::Add, NewDataLayer, NAME_None);
+	if (NewDataLayer != nullptr)
+	{
+		BroadcastDataLayerChanged(EDataLayerAction::Add, NewDataLayer, NAME_None);
 			SetParentDataLayer(NewDataLayer, Parameters.ParentDataLayer);
 		}
 	}
@@ -1125,20 +1126,20 @@ void UDataLayerEditorSubsystem::OnSelectionChanged()
 const TSet<TWeakObjectPtr<const UDataLayerInstance>>& UDataLayerEditorSubsystem::GetSelectedDataLayersFromEditorSelection() const
 {
 	if (bRebuildSelectedDataLayersFromEditorSelection)
-	{
+{
 		bRebuildSelectedDataLayersFromEditorSelection = false;
 
-		SelectedDataLayersFromEditorSelection.Reset();
-		TArray<AActor*> Actors;
-		GEditor->GetSelectedActors()->GetSelectedObjects<AActor>(Actors);
-		for (const AActor* Actor : Actors)
+	SelectedDataLayersFromEditorSelection.Reset();
+	TArray<AActor*> Actors;
+	GEditor->GetSelectedActors()->GetSelectedObjects<AActor>(Actors);
+	for (const AActor* Actor : Actors)
+	{
+		for (const UDataLayerInstance* DataLayerInstance : Actor->GetDataLayerInstances())
 		{
-			for (const UDataLayerInstance* DataLayerInstance : Actor->GetDataLayerInstances())
-			{
-				SelectedDataLayersFromEditorSelection.Add(DataLayerInstance);
-			}
+			SelectedDataLayersFromEditorSelection.Add(DataLayerInstance);
 		}
 	}
+}
 	return SelectedDataLayersFromEditorSelection;
 }
 
diff --git a/Engine/Source/Editor/UnrealEd/Private/Commandlets/WorldPartitionConvertCommandlet.cpp b/Engine/Source/Editor/UnrealEd/Private/Commandlets/WorldPartitionConvertCommandlet.cpp
--- a/Engine/Source/Editor/UnrealEd/Private/Commandlets/WorldPartitionConvertCommandlet.cpp
+++ b/Engine/Source/Editor/UnrealEd/Private/Commandlets/WorldPartitionConvertCommandlet.cpp
@@ -1054,7 +1054,7 @@ int32 UWorldPartitionConvertCommandlet::Main(const FString& Params)
 					}
 
 					// Convert Layers into DataLayers with DynamicallyLoaded flag disabled
-					if (Actor->IsValidForDataLayer())
+					if (Actor->SupportsDataLayer())
 					{
 						for (FName Layer : Actor->Layers)
 						{
diff --git a/Engine/Source/Runtime/Engine/Classes/Atmosphere/AtmosphericFog.h b/Engine/Source/Runtime/Engine/Classes/Atmosphere/AtmosphericFog.h
--- a/Engine/Source/Runtime/Engine/Classes/Atmosphere/AtmosphericFog.h
+++ b/Engine/Source/Runtime/Engine/Classes/Atmosphere/AtmosphericFog.h
@@ -18,6 +18,10 @@ class UE_DEPRECATED(4.26, "Please use the SkyAtmosphere actor instead.") AAtmosp
 	GENERATED_UCLASS_BODY()
 
 private:
+#if WITH_EDITOR
+	virtual bool ActorTypeSupportsDataLayer() const override { return true; }
+#endif
+
 	/** Main fog component */
 	UPROPERTY(BlueprintReadOnly, VisibleAnywhere, Category = Atmosphere, meta = (AllowPrivateAccess = "true"))
 	TObjectPtr<class UAtmosphericFogComponent> AtmosphericFogComponent;
@@ -29,11 +33,6 @@ private:
 #endif
 
 public:
-
-#if WITH_EDITOR
-	virtual bool SupportsDataLayer() const override { return true; }
-#endif
-
 	/** Returns AtmosphericFogComponent subobject **/
 	ENGINE_API class UAtmosphericFogComponent* GetAtmosphericFogComponent() { return AtmosphericFogComponent; }
 
diff --git a/Engine/Source/Runtime/Engine/Classes/Components/SkyAtmosphereComponent.h b/Engine/Source/Runtime/Engine/Classes/Components/SkyAtmosphereComponent.h
--- a/Engine/Source/Runtime/Engine/Classes/Components/SkyAtmosphereComponent.h
+++ b/Engine/Source/Runtime/Engine/Classes/Components/SkyAtmosphereComponent.h
@@ -273,6 +273,9 @@ class ASkyAtmosphere : public AInfo
 	GENERATED_UCLASS_BODY()
 
 private:
+#if WITH_EDITOR
+	virtual bool ActorTypeSupportsDataLayer() const override { return true; }
+#endif
 
 	UPROPERTY(BlueprintReadOnly, VisibleAnywhere, Category = Atmosphere, meta = (AllowPrivateAccess = "true"))
 	TObjectPtr<class USkyAtmosphereComponent> SkyAtmosphereComponent;
@@ -284,9 +287,6 @@ private:
 #endif
 
 public:
-#if WITH_EDITOR
-	virtual bool SupportsDataLayer() const override { return true; }
-#endif
 
 	/** Returns SkyAtmosphereComponent subobject */
 	ENGINE_API USkyAtmosphereComponent* GetComponent() const { return SkyAtmosphereComponent; }
diff --git a/Engine/Source/Runtime/Engine/Classes/Components/VolumetricCloudComponent.h b/Engine/Source/Runtime/Engine/Classes/Components/VolumetricCloudComponent.h
--- a/Engine/Source/Runtime/Engine/Classes/Components/VolumetricCloudComponent.h
+++ b/Engine/Source/Runtime/Engine/Classes/Components/VolumetricCloudComponent.h
@@ -204,9 +204,8 @@ private:
 	UPROPERTY(BlueprintReadOnly, VisibleAnywhere, Category = Atmosphere, meta = (AllowPrivateAccess = "true"))
 	TObjectPtr<class UVolumetricCloudComponent> VolumetricCloudComponent;
 
-public:
 #if WITH_EDITOR
-	virtual bool SupportsDataLayer() const override { return true; }
+	virtual bool ActorTypeSupportsDataLayer() const override { return true; }
 #endif
 
 };
diff --git a/Engine/Source/Runtime/Engine/Classes/Engine/ExponentialHeightFog.h b/Engine/Source/Runtime/Engine/Classes/Engine/ExponentialHeightFog.h
--- a/Engine/Source/Runtime/Engine/Classes/Engine/ExponentialHeightFog.h
+++ b/Engine/Source/Runtime/Engine/Classes/Engine/ExponentialHeightFog.h
@@ -19,14 +19,15 @@ class AExponentialHeightFog
 	virtual void GetLifetimeReplicatedProps(TArray< FLifetimeProperty > & OutLifetimeProps) const override;
 
 private:
+#if WITH_EDITOR
+	virtual bool ActorTypeSupportsDataLayer() const override { return true; }
+#endif
+
 	/** @todo document */
 	UPROPERTY(Category = ExponentialHeightFog, VisibleAnywhere, BlueprintReadOnly, meta = (AllowPrivateAccess = "true"))
 	TObjectPtr<class UExponentialHeightFogComponent> Component;
 
 public:
-#if WITH_EDITOR
-	virtual bool SupportsDataLayer() const override { return true; }
-#endif
 
 	/** replicated copy of ExponentialHeightFogComponent's bEnabled property */
 	UPROPERTY(replicatedUsing=OnRep_bEnabled)
diff --git a/Engine/Source/Runtime/Engine/Classes/Engine/SkyLight.h b/Engine/Source/Runtime/Engine/Classes/Engine/SkyLight.h
--- a/Engine/Source/Runtime/Engine/Classes/Engine/SkyLight.h
+++ b/Engine/Source/Runtime/Engine/Classes/Engine/SkyLight.h
@@ -15,14 +15,14 @@ class ENGINE_API ASkyLight : public AInfo
 	virtual void GetLifetimeReplicatedProps(TArray< FLifetimeProperty > & OutLifetimeProps) const override;
 
 private:
+#if WITH_EDITOR
+	virtual bool ActorTypeSupportsDataLayer() const override { return true; }
+#endif
+
 	/** @todo document */
 	UPROPERTY(Category = Light, VisibleAnywhere, BlueprintReadOnly, meta = (ExposeFunctionCategories = "Light,Rendering,Rendering|Components|SkyLight", AllowPrivateAccess = "true"))
 	TObjectPtr<class USkyLightComponent> LightComponent;
 public:
-#if WITH_EDITOR
-	virtual bool SupportsDataLayer() const override { return true; }
-#endif
-
 	/** replicated copy of LightComponent's bEnabled property */
 	UPROPERTY(replicatedUsing=OnRep_bEnabled)
 	uint32 bEnabled:1;
diff --git a/Engine/Source/Runtime/Engine/Classes/Engine/WindDirectionalSource.h b/Engine/Source/Runtime/Engine/Classes/Engine/WindDirectionalSource.h
--- a/Engine/Source/Runtime/Engine/Classes/Engine/WindDirectionalSource.h
+++ b/Engine/Source/Runtime/Engine/Classes/Engine/WindDirectionalSource.h
@@ -13,6 +13,10 @@ class ENGINE_API AWindDirectionalSource : public AInfo
 	GENERATED_UCLASS_BODY()
 
 private:
+#if WITH_EDITOR
+	virtual bool ActorTypeSupportsDataLayer() const override { return true; }
+#endif
+
 	UPROPERTY(Category = WindDirectionalSource, VisibleAnywhere, BlueprintReadOnly, meta = (AllowPrivateAccess = "true"))
 	TObjectPtr<class UWindDirectionalSourceComponent> Component;
 
@@ -25,10 +29,6 @@ public:
 	/** Returns Component subobject **/
 	class UWindDirectionalSourceComponent* GetComponent() const { return Component; }
 
-#if WITH_EDITOR
-	virtual bool SupportsDataLayer() const override { return true; }
-#endif
-
 #if WITH_EDITORONLY_DATA
 	/** Returns ArrowComponent subobject **/
 	class UArrowComponent* GetArrowComponent() const { return ArrowComponent; }
diff --git a/Engine/Source/Runtime/Engine/Classes/GameFramework/Actor.h b/Engine/Source/Runtime/Engine/Classes/GameFramework/Actor.h
--- a/Engine/Source/Runtime/Engine/Classes/GameFramework/Actor.h
+++ b/Engine/Source/Runtime/Engine/Classes/GameFramework/Actor.h
@@ -1177,17 +1177,19 @@ public:
 	//~=============================================================================
 	// DataLayers functions.
 #if WITH_EDITOR
+private:
+	virtual bool ActorTypeSupportsDataLayer() const { return true; }
+public:
 	bool AddDataLayer(const UDataLayerAsset* DataLayerAsset);
 	bool AddDataLayer(const UDataLayerInstance* DataLayerInstance);
 	bool RemoveDataLayer(const UDataLayerAsset* DataLayerAsset);
 	bool RemoveDataLayer(const UDataLayerInstance* DataLayerInstance);
 	bool RemoveAllDataLayers();
-	virtual bool SupportsDataLayer() const { return true; }
+	bool SupportsDataLayer() const;
 	
 	TArray<const UDataLayerInstance*> GetDataLayerInstancesForLevel() const;
 	TArray<FName> GetDataLayerInstanceNames() const;
 	bool IsPropertyChangedAffectingDataLayers(FPropertyChangedEvent& PropertyChangedEvent) const;
-	bool IsValidForDataLayer() const;
 	void FixupDataLayers(bool bRevertChangesOnLockedDataLayer = false);
 	static const FName GetDataLayerAssetsPropertyName() { return GET_MEMBER_NAME_CHECKED(AActor, DataLayerAssets); }
 	static const FName GetDataLayerPropertyName() { return GET_MEMBER_NAME_CHECKED(AActor, DataLayers); }
diff --git a/Engine/Source/Runtime/Engine/Classes/GameFramework/Info.h b/Engine/Source/Runtime/Engine/Classes/GameFramework/Info.h
--- a/Engine/Source/Runtime/Engine/Classes/GameFramework/Info.h
+++ b/Engine/Source/Runtime/Engine/Classes/GameFramework/Info.h
@@ -21,6 +21,11 @@ class AInfo : public AActor
 {
 	GENERATED_UCLASS_BODY()
 
+private:
+#if WITH_EDITOR
+	virtual bool ActorTypeSupportsDataLayer() const override { return false; }
+#endif
+
 #if WITH_EDITORONLY_DATA
 private:
 	/** Billboard Component displayed in editor */
@@ -34,7 +39,6 @@ public:
 
 #if WITH_EDITOR
 	virtual bool CanChangeIsSpatiallyLoadedFlag() const override { return false; }
-	virtual bool SupportsDataLayer() const override { return false; }
 #endif
 
 public:
diff --git a/Engine/Source/Runtime/Engine/Private/Actor.cpp b/Engine/Source/Runtime/Engine/Private/Actor.cpp
--- a/Engine/Source/Runtime/Engine/Private/Actor.cpp
+++ b/Engine/Source/Runtime/Engine/Private/Actor.cpp
@@ -5960,7 +5960,7 @@ TArray<const UDataLayerInstance*> AActor::GetDataLayerInstancesInternal(bool bUs
 	}
 
 #if WITH_EDITOR
-	if (SupportsDataLayer() && IsValidForDataLayer())
+	if (SupportsDataLayer())
 	{
 		TArray<const UDataLayerInstance*> DataLayerInstances;
 		ULevel* OuterLevel = bUseLevelContext ? GetLevel() : nullptr;
diff --git a/Engine/Source/Runtime/Engine/Private/ActorEditor.cpp b/Engine/Source/Runtime/Engine/Private/ActorEditor.cpp
--- a/Engine/Source/Runtime/Engine/Private/ActorEditor.cpp
+++ b/Engine/Source/Runtime/Engine/Private/ActorEditor.cpp
@@ -1335,6 +1335,10 @@ EDataValidationResult AActor::IsDataValid(TArray<FText>& ValidationErrors)
 
 bool AActor::AddDataLayer(const UDataLayerInstance* DataLayerInstance)
 {
+	if (!SupportsDataLayer())
+	{
+		return false;
+	}
 	return DataLayerInstance->AddActor(this);
 }
 
@@ -1396,7 +1400,7 @@ TArray<FName> AActor::GetDataLayerInstanceNames() const
 }
 
 TArray<const UDataLayerInstance*> AActor::GetDataLayerInstancesForLevel() const
-{
+	{
 	const bool bUseLevelContext = true;
 	return GetDataLayerInstancesInternal(bUseLevelContext);
 }
@@ -1421,62 +1425,62 @@ void AActor::FixupDataLayers(bool bRevertChangesOnLockedDataLayer /*= false*/)
 
 			ULevel* Level = GetLevel();
 
-			if (bRevertChangesOnLockedDataLayer)
-			{
-				// Since it's not possible to prevent changes of particular elements of an array, rollback change on locked DataLayers.
-				TSet<const UDataLayerAsset*> PreEdit(PreEditChangeDataLayers);
-				TSet<const UDataLayerAsset*> PostEdit(DataLayerAssets);
+				if (bRevertChangesOnLockedDataLayer)
+				{
+					// Since it's not possible to prevent changes of particular elements of an array, rollback change on locked DataLayers.
+					TSet<const UDataLayerAsset*> PreEdit(PreEditChangeDataLayers);
+					TSet<const UDataLayerAsset*> PostEdit(DataLayerAssets);
 
 				auto DifferenceContainsLockedDataLayers = [DataLayerSubsystem, Level](const TSet<const UDataLayerAsset*>& A, const TSet<const UDataLayerAsset*>& B)
-				{
-					TSet<const UDataLayerAsset*> Diff = A.Difference(B);
-					for (const UDataLayerAsset* DataLayerAsset : Diff)
 					{
+						TSet<const UDataLayerAsset*> Diff = A.Difference(B);
+						for (const UDataLayerAsset* DataLayerAsset : Diff)
+						{
 						// We pass Actor Level when resolving the DataLayerInstance as we do the fixup relative to this level
 						const UDataLayerInstance* DataLayerInstance = DataLayerSubsystem->GetDataLayerInstance(DataLayerAsset, Level);
-						if (DataLayerInstance && DataLayerInstance->IsLocked())
-						{
-							return true;
+							if (DataLayerInstance && DataLayerInstance->IsLocked())
+							{
+								return true;
+							}
 						}
-					}
-					return false;
-				};
+						return false;
+					};
 					
-				if (DifferenceContainsLockedDataLayers(PreEdit, PostEdit) || 
-					DifferenceContainsLockedDataLayers(PostEdit, PreEdit))
-				{
-					DataLayerAssets = PreEditChangeDataLayers;
+					if (DifferenceContainsLockedDataLayers(PreEdit, PostEdit) || 
+						DifferenceContainsLockedDataLayers(PostEdit, PreEdit))
+					{
+						DataLayerAssets = PreEditChangeDataLayers;
+					}
 				}
-			}
 
-			PRAGMA_DISABLE_DEPRECATION_WARNINGS
+				PRAGMA_DISABLE_DEPRECATION_WARNINGS
 			auto CleanupDataLayers = [this, DataLayerSubsystem, Level](auto& DataLayerArray)
-			{
-				using ArrayType = typename TRemoveReference<decltype(DataLayerArray)>::Type;
-					
-				ArrayType ExistingDataLayer;
-				for (int32 Index = 0; Index < DataLayerArray.Num();)
 				{
+					using ArrayType = typename TRemoveReference<decltype(DataLayerArray)>::Type;
+					
+					ArrayType ExistingDataLayer;
+					for (int32 Index = 0; Index < DataLayerArray.Num();)
+					{
 					// We pass Actor Level when resolving the DataLayerInstance as we do the fixup relative to this level
-					auto& DataLayer = DataLayerArray[Index];
+						auto& DataLayer = DataLayerArray[Index];
 					if (!DataLayerSubsystem->GetDataLayerInstance(DataLayer, Level) || ExistingDataLayer.Contains(DataLayer))
-					{
-						DataLayerArray.RemoveAtSwap(Index);
-					}
-					else
-					{
-						ExistingDataLayer.Add(DataLayer);
-						++Index;
+						{
+							DataLayerArray.RemoveAtSwap(Index);
+						}
+						else
+						{
+							ExistingDataLayer.Add(DataLayer);
+							++Index;
+						}
 					}
-				}
-			};
+				};
 
-			CleanupDataLayers(DataLayerAssets);
-			CleanupDataLayers(DataLayers);
-			PRAGMA_ENABLE_DEPRECATION_WARNINGS
+				CleanupDataLayers(DataLayerAssets);
+				CleanupDataLayers(DataLayers);
+				PRAGMA_ENABLE_DEPRECATION_WARNINGS
+			}
 		}
 	}
-}
 
 bool AActor::IsPropertyChangedAffectingDataLayers(FPropertyChangedEvent& PropertyChangedEvent) const
 {
@@ -1496,21 +1500,11 @@ bool AActor::IsPropertyChangedAffectingDataLayers(FPropertyChangedEvent& Propert
 	return false;
 }
 
-bool AActor::IsValidForDataLayer() const
+bool AActor::SupportsDataLayer() const
 {
-	UWorld* World = GetWorld();
-	if (!World)
-	{
-		return false;
-	}
-
-	const bool bIsPartitionedActor = UWorld::HasSubsystem<UWorldPartitionSubsystem>(World);
-	const bool bIsInEditorWorld = World->WorldType == EWorldType::Editor;
-	const bool bIsBuilderBrush = FActorEditorUtils::IsABuilderBrush(this);
-	const bool bIsHidden = GetClass()->GetDefaultObject<AActor>()->bHiddenEd;
-	const bool bIsValid = !bIsHidden && !bIsBuilderBrush && bIsInEditorWorld && bIsPartitionedActor;
-
-	return bIsValid;
+	return (ActorTypeSupportsDataLayer() &&
+			!FActorEditorUtils::IsABuilderBrush(this) &&
+			!GetClass()->GetDefaultObject<AActor>()->bHiddenEd);
 }
 
 //~ Begin Deprecated
@@ -1519,7 +1513,7 @@ PRAGMA_DISABLE_DEPRECATION_WARNINGS
 
 bool AActor::AddDataLayer(const FActorDataLayer& ActorDataLayer)
 {
-	if (!ContainsDataLayer(ActorDataLayer))
+	if (SupportsDataLayer() && !ContainsDataLayer(ActorDataLayer))
 	{
 		Modify();
 		DataLayers.Add(ActorDataLayer);
diff --git a/Engine/Source/Runtime/Engine/Private/LevelInstance/LevelInstanceSubsystem.cpp b/Engine/Source/Runtime/Engine/Private/LevelInstance/LevelInstanceSubsystem.cpp
--- a/Engine/Source/Runtime/Engine/Private/LevelInstance/LevelInstanceSubsystem.cpp
+++ b/Engine/Source/Runtime/Engine/Private/LevelInstance/LevelInstanceSubsystem.cpp
@@ -1199,7 +1199,7 @@ void ULevelInstanceSubsystem::BreakLevelInstance_Impl(ILevelInstanceInterface* L
 					}
 
 					// Apply the same data layer settings to the actors to move out
-					if (Actor->SupportsDataLayer() && Actor->IsValidForDataLayer())
+					if (Actor->SupportsDataLayer())
 					{
 						for (const UDataLayerInstance* DataLayerInstance : LevelInstanceDataLayerInstances)
 						{
diff --git a/Engine/Source/Runtime/Engine/Private/LightWeightInstanceSubsystem.cpp b/Engine/Source/Runtime/Engine/Private/LightWeightInstanceSubsystem.cpp
--- a/Engine/Source/Runtime/Engine/Private/LightWeightInstanceSubsystem.cpp
+++ b/Engine/Source/Runtime/Engine/Private/LightWeightInstanceSubsystem.cpp
@@ -120,7 +120,7 @@ ALightWeightInstanceManager* FLightWeightInstanceSubsystem::FindLightWeightInsta
 	return nullptr;
 }
 
-ALightWeightInstanceManager* FLightWeightInstanceSubsystem::FindOrAddLightWeightInstanceManager(UClass* ActorClass, const UDataLayerInstance* Layer, UWorld* World)
+ALightWeightInstanceManager* FLightWeightInstanceSubsystem::FindOrAddLightWeightInstanceManager(UClass* ActorClass, const UDataLayerInstance* DataLayer, UWorld* World)
 {
 	if (ActorClass == nullptr || World == nullptr)
 	{
@@ -133,7 +133,7 @@ ALightWeightInstanceManager* FLightWeightInstanceSubsystem::FindOrAddLightWeight
 		if (InstanceManager->GetRepresentedClass() == ActorClass)
 		{
 #if WITH_EDITOR
-			if (!Layer || (InstanceManager->SupportsDataLayer() && InstanceManager->ContainsDataLayer(Layer)))
+			if (!DataLayer || (InstanceManager->SupportsDataLayer() && InstanceManager->ContainsDataLayer(DataLayer)))
 #endif // WITH_EDITOR
 			{
 				return InstanceManager;
@@ -158,12 +158,12 @@ ALightWeightInstanceManager* FLightWeightInstanceSubsystem::FindOrAddLightWeight
 	NewInstanceManager->SetRepresentedClass(ActorClass);
 
 #if WITH_EDITOR
-	// Add the new manager to the layer
-	if (Layer)
+	// Add the new manager to the DataLayer
+	if (DataLayer)
 	{
-		ensure(NewInstanceManager->IsValidForDataLayer());
+		ensure(NewInstanceManager->SupportsDataLayer());
 
-		NewInstanceManager->AddDataLayer(Layer);
+		NewInstanceManager->AddDataLayer(DataLayer);
 	}
 #endif // WITH_EDITOR
 
diff --git a/Engine/Source/Runtime/Engine/Public/LocationVolume.h b/Engine/Source/Runtime/Engine/Public/LocationVolume.h
--- a/Engine/Source/Runtime/Engine/Public/LocationVolume.h
+++ b/Engine/Source/Runtime/Engine/Public/LocationVolume.h
@@ -30,7 +30,6 @@ public:
 #if WITH_EDITOR
 	virtual void PostRegisterAllComponents();
 	virtual bool CanChangeIsSpatiallyLoadedFlag() const override { return false; }
-	virtual bool SupportsDataLayer() const override { return false; }
 	//~ End AActor Interface
 
 	//~ Begin IWorldPartitionActorLoaderInterface interface
@@ -51,6 +50,7 @@ public:
 
 #if WITH_EDITOR
 private:
+	virtual bool ActorTypeSupportsDataLayer() const override { return false; }
 	FLoaderAdapterActor* WorldPartitionActorLoader;
 #endif
 };
\ No newline at end of file
diff --git a/Engine/Source/Runtime/Engine/Public/WorldPartition/WorldPartitionMiniMap.h b/Engine/Source/Runtime/Engine/Public/WorldPartition/WorldPartitionMiniMap.h
--- a/Engine/Source/Runtime/Engine/Public/WorldPartition/WorldPartitionMiniMap.h
+++ b/Engine/Source/Runtime/Engine/Public/WorldPartition/WorldPartitionMiniMap.h
@@ -14,6 +14,11 @@ class ENGINE_API AWorldPartitionMiniMap : public AInfo
 {
 	GENERATED_BODY()
 
+private:
+#if WITH_EDITOR
+	virtual bool ActorTypeSupportsDataLayer() const final { return false; }
+#endif
+
 public:
 	AWorldPartitionMiniMap(const FObjectInitializer& ObjectInitializer = FObjectInitializer::Get());
 
@@ -23,7 +28,6 @@ public:
 	virtual void PostEditChangeProperty(FPropertyChangedEvent& PropertyChangedEvent) override;
 	virtual bool IsLockLocation() const { return true; }
 	virtual bool IsUserManaged() const final { return false; }
-	virtual bool SupportsDataLayer() const final { return false; }
 	virtual void CheckForErrors() override;
 #endif
 
diff --git a/Engine/Source/Runtime/Engine/Public/WorldPartition/WorldPartitionVolume.h b/Engine/Source/Runtime/Engine/Public/WorldPartition/WorldPartitionVolume.h
--- a/Engine/Source/Runtime/Engine/Public/WorldPartition/WorldPartitionVolume.h
+++ b/Engine/Source/Runtime/Engine/Public/WorldPartition/WorldPartitionVolume.h
@@ -13,12 +13,16 @@ class ENGINE_API ADEPRECATED_WorldPartitionVolume : public AVolume
 {
 	GENERATED_UCLASS_BODY()
 
+private:
+#if WITH_EDITOR
+	virtual bool ActorTypeSupportsDataLayer() const override { return false; }
+#endif
+
 public:
 	//~ Begin AActor Interface
 	virtual bool IsEditorOnly() const override { return true; }
 #if WITH_EDITOR
 	virtual bool CanChangeIsSpatiallyLoadedFlag() const override { return false; }
-	virtual bool SupportsDataLayer() const override { return false; }
 #endif
 	//~ End AActor Interface
 };
\ No newline at end of file
