Dissallow  Level Instance Pivot actor to have Data Layers
#rnx

diff --git a/Engine/Source/Editor/DataLayerEditor/Private/DataLayer/DataLayerEditorSubsystem.cpp b/Engine/Source/Editor/DataLayerEditor/Private/DataLayer/DataLayerEditorSubsystem.cpp
--- a/Engine/Source/Editor/DataLayerEditor/Private/DataLayer/DataLayerEditorSubsystem.cpp
+++ b/Engine/Source/Editor/DataLayerEditor/Private/DataLayer/DataLayerEditorSubsystem.cpp
@@ -180,30 +180,30 @@ void UDataLayerEditorSubsystem::OnExecuteActorEditorContextAction(UWorld* InWorl
 	UE_CLOG(!InWorld, LogDataLayerEditorSubsystem, Error, TEXT("%s - Failed because world in null."), ANSI_TO_TCHAR(__FUNCTION__));
 	if (UDataLayerSubsystem* DataLayerSubsystem = UWorld::GetSubsystem<UDataLayerSubsystem>(InWorld))
 	{
-	switch (InType)
-	{
-	case EActorEditorContextAction::ApplyContext:
-		check(InActor && InActor->GetWorld() == InWorld);
+		switch (InType)
 		{
+		case EActorEditorContextAction::ApplyContext:
+			check(InActor && InActor->GetWorld() == InWorld);
+			{
 				AddActorToDataLayers(InActor, DataLayerSubsystem->GetActorEditorContextDataLayers());
-		}
-		break;
-	case EActorEditorContextAction::ResetContext:
+			}
+			break;
+		case EActorEditorContextAction::ResetContext:
 			for (UDataLayerInstance* DataLayer : DataLayerSubsystem->GetActorEditorContextDataLayers())
-		{
-			RemoveFromActorEditorContext(DataLayer);
-		}
-		break;
-	case EActorEditorContextAction::PushContext:
+			{
+				RemoveFromActorEditorContext(DataLayer);
+			}
+			break;
+		case EActorEditorContextAction::PushContext:
 			DataLayerSubsystem->PushActorEditorContext();
-		BroadcastDataLayerChanged(EDataLayerAction::Reset, NULL, NAME_None);
-		break;
-	case EActorEditorContextAction::PopContext:
+			BroadcastDataLayerChanged(EDataLayerAction::Reset, NULL, NAME_None);
+			break;
+		case EActorEditorContextAction::PopContext:
 			DataLayerSubsystem->PopActorEditorContext();
-		BroadcastDataLayerChanged(EDataLayerAction::Reset, NULL, NAME_None);
-		break;
+			BroadcastDataLayerChanged(EDataLayerAction::Reset, NULL, NAME_None);
+			break;
+		}
 	}
-}
 }	
 
 bool UDataLayerEditorSubsystem::GetActorEditorContextDisplayInfo(UWorld* InWorld, FActorEditorContextClientDisplayInfo& OutDiplayInfo) const
@@ -211,11 +211,11 @@ bool UDataLayerEditorSubsystem::GetActorEditorContextDisplayInfo(UWorld* InWorld
 	if (UDataLayerSubsystem* DataLayerSubsystem = UWorld::GetSubsystem<UDataLayerSubsystem>(InWorld))
 	{
 		if (!DataLayerSubsystem->GetActorEditorContextDataLayers().IsEmpty())
-	{
+		{
 			OutDiplayInfo.Title = TEXT("Data Layers");
 			OutDiplayInfo.Brush = FAppStyle::GetBrush(TEXT("DataLayer.Editor"));
-		return true;
-	}
+			return true;
+		}
 	}
 	return false;
 }
diff --git a/Engine/Source/Runtime/Engine/Public/LevelInstance/LevelInstanceEditorPivotActor.h b/Engine/Source/Runtime/Engine/Public/LevelInstance/LevelInstanceEditorPivotActor.h
--- a/Engine/Source/Runtime/Engine/Public/LevelInstance/LevelInstanceEditorPivotActor.h
+++ b/Engine/Source/Runtime/Engine/Public/LevelInstance/LevelInstanceEditorPivotActor.h
@@ -1,25 +1,27 @@
-// Copyright Epic Games, Inc. All Rights Reserved.
-
-#pragma once
-
-#include "CoreMinimal.h"
-#include "UObject/ObjectMacros.h"
-#include "GameFramework/Actor.h"
-#include "LevelInstance/LevelInstanceTypes.h"
-#include "LevelInstance/LevelInstanceEditorPivotInterface.h"
-#include "LevelInstanceEditorPivotActor.generated.h"
-
-UCLASS(transient, notplaceable, hidecategories=(Rendering,Replication,Collision,Partition,Input,HLOD,Actor,Cooking))
-class ENGINE_API ALevelInstancePivot : public AActor, public ILevelInstanceEditorPivotInterface
-{
-	GENERATED_UCLASS_BODY()
-
-#if WITH_EDITOR
-public:
-	virtual bool CanDeleteSelectedActor(FText& OutReason) const override { return false; }
-	virtual void PostEditMove(bool bFinished) override;
-	virtual void PostEditChangeProperty(FPropertyChangedEvent& PropertyChangedEvent) override;
-	virtual void PostEditUndo() override;
-	virtual bool ShouldExport() override { return false; }
-#endif
+// Copyright Epic Games, Inc. All Rights Reserved.
+
+#pragma once
+
+#include "CoreMinimal.h"
+#include "UObject/ObjectMacros.h"
+#include "GameFramework/Actor.h"
+#include "LevelInstance/LevelInstanceTypes.h"
+#include "LevelInstance/LevelInstanceEditorPivotInterface.h"
+#include "LevelInstanceEditorPivotActor.generated.h"
+
+UCLASS(transient, notplaceable, hidecategories=(Rendering,Replication,Collision,Partition,Input,HLOD,Actor,Cooking))
+class ENGINE_API ALevelInstancePivot : public AActor, public ILevelInstanceEditorPivotInterface
+{
+	GENERATED_UCLASS_BODY()
+
+#if WITH_EDITOR
+public:
+	virtual bool CanDeleteSelectedActor(FText& OutReason) const override { return false; }
+	virtual void PostEditMove(bool bFinished) override;
+	virtual void PostEditChangeProperty(FPropertyChangedEvent& PropertyChangedEvent) override;
+	virtual void PostEditUndo() override;
+	virtual bool ShouldExport() override { return false; }
+private:
+	virtual bool ActorTypeSupportsDataLayer() const { return false; }
+#endif
 };
\ No newline at end of file
