Data Layer: fix old data layer assets debug color when loading.

#rnx

diff --git a/Engine/Source/Editor/DataLayerEditor/Private/DataLayer/DataLayerFactory.cpp b/Engine/Source/Editor/DataLayerEditor/Private/DataLayer/DataLayerFactory.cpp
--- a/Engine/Source/Editor/DataLayerEditor/Private/DataLayer/DataLayerFactory.cpp
+++ b/Engine/Source/Editor/DataLayerEditor/Private/DataLayer/DataLayerFactory.cpp
@@ -1,13 +1,10 @@
 // Copyright Epic Games, Inc. All Rights Reserved.
 
 #include "DataLayer/DataLayerFactory.h"
-
 #include "WorldPartition/DataLayer/DataLayerAsset.h"
-#include "Math/RandomStream.h"
 
 UDataLayerFactory::UDataLayerFactory(const FObjectInitializer& ObjectInitializer)
 	: Super(ObjectInitializer)
-
 {
 	SupportedClass = UDataLayerAsset::StaticClass();
 	bCreateNew = true;
@@ -17,13 +14,7 @@ UDataLayerFactory::UDataLayerFactory(const FObjectInitializer& ObjectInitializer
 
 UObject* UDataLayerFactory::FactoryCreateNew(UClass* Class, UObject* InParent, FName InName, EObjectFlags Flags, UObject* Context, FFeedbackContext* Warn)
 {
-	UDataLayerAsset* DataLayerAsset =   NewObject<UDataLayerAsset>(InParent, InName, Flags);
-
-	FRandomStream RandomStream(FName(DataLayerAsset->GetFullName()));
-	const uint8 R = (uint8)(RandomStream.GetFraction() * 255.f);
-	const uint8 G = (uint8)(RandomStream.GetFraction() * 255.f);
-	const uint8 B = (uint8)(RandomStream.GetFraction() * 255.f);
-	DataLayerAsset->SetDebugColor(FColor(R, G, B));
-
+	UDataLayerAsset* DataLayerAsset = NewObject<UDataLayerAsset>(InParent, InName, Flags);
+	DataLayerAsset->SetDebugColor(FColor::MakeRandomSeededColor(GetTypeHash(DataLayerAsset->GetFullName())));
 	return DataLayerAsset;
 }
\ No newline at end of file
diff --git a/Engine/Source/Runtime/Core/Private/Math/Color.cpp b/Engine/Source/Runtime/Core/Private/Math/Color.cpp
--- a/Engine/Source/Runtime/Core/Private/Math/Color.cpp
+++ b/Engine/Source/Runtime/Core/Private/Math/Color.cpp
@@ -8,6 +8,7 @@
 #include "Math/Vector.h"
 #include "Math/Vector4.h"
 #include "Math/Float16Color.h"
+#include "Math/RandomStream.h"
 
 // Common colors.
 const FLinearColor FLinearColor::White(1.f,1.f,1.f);
@@ -513,6 +514,22 @@ FLinearColor FLinearColor::MakeFromColorTemperature( float Temp )
 	return FLinearColor(R,G,B);
 }
 
+FLinearColor FLinearColor::MakeRandomSeededColor(int32 Seed)
+{
+	FRandomStream RandomStream(Seed);
+
+	float R = RandomStream.GetFraction();
+	float G = RandomStream.GetFraction();
+	float B = RandomStream.GetFraction();
+
+	return FLinearColor(R,G,B);
+}
+
+FColor FColor::MakeRandomSeededColor(int32 Seed)
+{
+	return FLinearColor::MakeRandomSeededColor( Seed ).ToFColor( true );
+}
+
 FColor FColor::MakeFromColorTemperature( float Temp )
 {
 	return FLinearColor::MakeFromColorTemperature( Temp ).ToFColor( true );
diff --git a/Engine/Source/Runtime/Core/Public/Math/Color.h b/Engine/Source/Runtime/Core/Public/Math/Color.h
--- a/Engine/Source/Runtime/Core/Public/Math/Color.h
+++ b/Engine/Source/Runtime/Core/Public/Math/Color.h
@@ -288,6 +288,11 @@ struct FLinearColor
 	*/
 	static CORE_API FLinearColor MakeFromColorTemperature( float Temp );
 
+	/**
+	* Makes a random color based on a seed.
+	*/
+	static CORE_API FLinearColor MakeRandomSeededColor(int32 Seed);
+
 	/**
 	 * Euclidean distance between two points.
 	 */
@@ -568,6 +573,11 @@ public:
 	*/
 	static CORE_API FColor MakeFromColorTemperature( float Temp );
 
+	/**
+	* Makes a random color based on a seed.
+	*/
+	static CORE_API FColor MakeRandomSeededColor(int32 Seed);
+
 	/**
 	* Conversions to/from GPU UNorm floats, U8, U16
 	* matches convention of FColor FLinearColor::QuantizeRound
diff --git a/Engine/Source/Runtime/Engine/Private/WorldPartition/DataLayer/DataLayer.cpp b/Engine/Source/Runtime/Engine/Private/WorldPartition/DataLayer/DataLayer.cpp
--- a/Engine/Source/Runtime/Engine/Private/WorldPartition/DataLayer/DataLayer.cpp
+++ b/Engine/Source/Runtime/Engine/Private/WorldPartition/DataLayer/DataLayer.cpp
@@ -26,8 +26,7 @@ UDEPRECATED_DataLayer::UDEPRECATED_DataLayer(const FObjectInitializer& ObjectIni
 , bIsRuntime(false)
 , InitialRuntimeState(EDataLayerRuntimeState::Unloaded)
 , DebugColor(FColor::Black)
-{
-}
+{}
 
 void UDEPRECATED_DataLayer::PostLoad()
 {
@@ -47,11 +46,7 @@ void UDEPRECATED_DataLayer::PostLoad()
 
 	if (DebugColor == FColor::Black)
 	{
-		FRandomStream RandomStream(GetFName());
-		const uint8 R = (uint8)(RandomStream.GetFraction() * 255.f);
-		const uint8 G = (uint8)(RandomStream.GetFraction() * 255.f);
-		const uint8 B = (uint8)(RandomStream.GetFraction() * 255.f);
-		DebugColor = FColor(R, G, B);
+		DebugColor = FColor::MakeRandomSeededColor(GetTypeHash(GetName()));
 	}
 #endif
 
diff --git a/Engine/Source/Runtime/Engine/Private/WorldPartition/DataLayer/DataLayerAsset.cpp b/Engine/Source/Runtime/Engine/Private/WorldPartition/DataLayer/DataLayerAsset.cpp
--- a/Engine/Source/Runtime/Engine/Private/WorldPartition/DataLayer/DataLayerAsset.cpp
+++ b/Engine/Source/Runtime/Engine/Private/WorldPartition/DataLayer/DataLayerAsset.cpp
@@ -3,9 +3,19 @@
 #include "WorldPartition/DataLayer/DataLayerAsset.h"
 
 UDataLayerAsset::UDataLayerAsset(const FObjectInitializer& ObjectInitializer)
-	: Super(ObjectInitializer),
-	DataLayerType(EDataLayerType::Editor),
-	DebugColor(FColor::Black)
+	: Super(ObjectInitializer)
+	, DataLayerType(EDataLayerType::Editor)
+	, DebugColor(FColor::Black)
+{}
+
+#if WITH_EDITOR
+void UDataLayerAsset::PostLoad()
 {
+	if (DebugColor == FColor::Black)
+	{
+		DebugColor = FColor::MakeRandomSeededColor(GetTypeHash(GetName()));
+	}
 
+	Super::PostLoad();
 }
+#endif
\ No newline at end of file
diff --git a/Engine/Source/Runtime/Engine/Private/WorldPartition/DataLayer/DeprecatedDataLayerInstance.cpp b/Engine/Source/Runtime/Engine/Private/WorldPartition/DataLayer/DeprecatedDataLayerInstance.cpp
--- a/Engine/Source/Runtime/Engine/Private/WorldPartition/DataLayer/DeprecatedDataLayerInstance.cpp
+++ b/Engine/Source/Runtime/Engine/Private/WorldPartition/DataLayer/DeprecatedDataLayerInstance.cpp
@@ -30,11 +30,17 @@ void UDeprecatedDataLayerInstance::OnCreated()
 
 	DeprecatedDataLayerFName = TEXT("");
 
-	FRandomStream RandomStream(GetDataLayerFName());
-	const uint8 R = (uint8)(RandomStream.GetFraction() * 255.f);
-	const uint8 G = (uint8)(RandomStream.GetFraction() * 255.f);
-	const uint8 B = (uint8)(RandomStream.GetFraction() * 255.f);
-	DebugColor = FColor(R, G, B);
+	DebugColor = FColor::MakeRandomSeededColor(GetTypeHash(GetDataLayerFName().ToString()));
+}
+
+void UDeprecatedDataLayerInstance::PostLoad()
+{
+	if (DebugColor == FColor::Black)
+	{
+		DebugColor = FColor::MakeRandomSeededColor(GetTypeHash(GetDataLayerFName().ToString()));
+	}
+
+	Super::PostLoad();
 }
 
 FName UDeprecatedDataLayerInstance::MakeName(const UDEPRECATED_DataLayer* DeprecatedDataLayer)
diff --git a/Engine/Source/Runtime/Engine/Public/WorldPartition/DataLayer/DataLayerAsset.h b/Engine/Source/Runtime/Engine/Public/WorldPartition/DataLayer/DataLayerAsset.h
--- a/Engine/Source/Runtime/Engine/Public/WorldPartition/DataLayer/DataLayerAsset.h
+++ b/Engine/Source/Runtime/Engine/Public/WorldPartition/DataLayer/DataLayerAsset.h
@@ -18,8 +18,12 @@ class ENGINE_API UDataLayerAsset : public UObject
 
 	friend class UDataLayerConversionInfo;
 
-public:
 #if WITH_EDITOR
+	//~ Begin UObject Interface
+	virtual void PostLoad() override;
+	//~ End UObject Interface
+
+public:
 	void SetType(EDataLayerType Type) { DataLayerType = Type; }
 	void SetDebugColor(FColor InDebugColor) { DebugColor = InDebugColor; }
 #endif
diff --git a/Engine/Source/Runtime/Engine/Public/WorldPartition/DataLayer/DeprecatedDataLayerInstance.h b/Engine/Source/Runtime/Engine/Public/WorldPartition/DataLayer/DeprecatedDataLayerInstance.h
--- a/Engine/Source/Runtime/Engine/Public/WorldPartition/DataLayer/DeprecatedDataLayerInstance.h
+++ b/Engine/Source/Runtime/Engine/Public/WorldPartition/DataLayer/DeprecatedDataLayerInstance.h
@@ -22,6 +22,12 @@ class ENGINE_API UDeprecatedDataLayerInstance final : public UDataLayerInstance
 	friend class UDataLayerToAssetCommandletContext;
 	friend class UDataLayerToAssetCommandlet;
 
+#if WITH_EDITOR
+	//~ Begin UObject Interface
+	virtual void PostLoad() override;
+	//~ End UObject Interface
+#endif
+
 public:
 #if WITH_EDITOR
 	static FName MakeName();
