Fixed FDataLayerMode::RegisterContextMenu to be static as it is registered once : use USceneOutlinerMenuContext to get the associated FDataLayerMode instance.

diff --git a/Engine/Source/Editor/DataLayerEditor/Private/DataLayer/DataLayerMode.cpp b/Engine/Source/Editor/DataLayerEditor/Private/DataLayer/DataLayerMode.cpp
--- a/Engine/Source/Editor/DataLayerEditor/Private/DataLayer/DataLayerMode.cpp
+++ b/Engine/Source/Editor/DataLayerEditor/Private/DataLayer/DataLayerMode.cpp
@@ -183,8 +183,6 @@ FDataLayerMode::FDataLayerMode(const FDataLayerModeParams& Params)
 
 FDataLayerMode::~FDataLayerMode()
 {
-	UnregisterContextMenu();
-
 	USelection::SelectionChangedEvent.RemoveAll(this);
 	USelection::SelectObjectEvent.RemoveAll(this);
 }
@@ -980,7 +978,7 @@ void FDataLayerMode::RegisterContextMenu()
 	{
 		UToolMenu* Menu = ToolMenus->RegisterMenu(DefaultContextBaseMenuName);
 
-		Menu->AddDynamicSection("DataLayerDynamicSection", FNewToolMenuDelegate::CreateLambda([this](UToolMenu* InMenu)
+		Menu->AddDynamicSection("DataLayerDynamicSection", FNewToolMenuDelegate::CreateLambda([](UToolMenu* InMenu)
 		{
 			USceneOutlinerMenuContext* Context = InMenu->FindContext<USceneOutlinerMenuContext>();
 			if (!Context || !Context->SceneOutliner.IsValid())
@@ -989,12 +987,14 @@ void FDataLayerMode::RegisterContextMenu()
 			}
 
 			SSceneOutliner* SceneOutliner = Context->SceneOutliner.Pin().Get();
-			TArray<UDataLayerInstance*> SelectedDataLayers = GetSelectedDataLayers(SceneOutliner);
+			FDataLayerMode* Mode = const_cast<FDataLayerMode*>(static_cast<const FDataLayerMode*>(SceneOutliner->GetMode()));
+			check(Mode);
+			TArray<UDataLayerInstance*> SelectedDataLayers = Mode->GetSelectedDataLayers(SceneOutliner);
 			const bool bSelectedDataLayersContainsLocked = !!SelectedDataLayers.FindByPredicate([](const UDataLayerInstance* DataLayer) { return DataLayer->IsLocked(); });
 
 			bool bHasActorEditorContextDataLayers = false;			
 			TArray<const UDataLayerInstance*> AllDataLayers;
-			if (const UDataLayerSubsystem* DataLayerSubsystem = UWorld::GetSubsystem<UDataLayerSubsystem>(GetOwningWorld()))
+			if (const UDataLayerSubsystem* DataLayerSubsystem = UWorld::GetSubsystem<UDataLayerSubsystem>(Mode->GetOwningWorld()))
 			{
 				DataLayerSubsystem->ForEachDataLayer([&AllDataLayers](UDataLayerInstance* DataLayerInstance)
 				{
@@ -1008,14 +1008,14 @@ void FDataLayerMode::RegisterContextMenu()
 			{
 				FToolMenuSection& Section = InMenu->AddSection("DataLayers", LOCTEXT("DataLayers", "Data Layers"));
 				
-				auto CreateNewDataLayer = [this, SceneOutliner](UDataLayerInstance* InParentDataLayer = nullptr, const TSet<const UDataLayerAsset*>& InDataLayerAssets = {})
+				auto CreateNewDataLayer = [SceneOutliner, Mode](UDataLayerInstance* InParentDataLayer = nullptr, const TSet<const UDataLayerAsset*>& InDataLayerAssets = {})
 				{
 					TSet<const UDataLayerAsset*> DataLayerAssets = InDataLayerAssets;
 					DataLayerAssets.Remove(nullptr);
 
 					if (DataLayerAssets.IsEmpty())
 					{
-						if (UDataLayerAsset* DataLayerAsset = PromptDataLayerAssetSelection())
+						if (UDataLayerAsset* DataLayerAsset = Mode->PromptDataLayerAssetSelection())
 						{
 							DataLayerAssets.Add(DataLayerAsset);
 						}
@@ -1024,16 +1024,16 @@ void FDataLayerMode::RegisterContextMenu()
 					for (const UDataLayerAsset* DataLayerAsset : DataLayerAssets)
 					{
 						const FScopedTransaction Transaction(LOCTEXT("CreateNewDataLayer", "Create New Data Layer"));
-						SelectedDataLayersSet.Empty();
-						SelectedDataLayerActors.Empty();
+						Mode->SelectedDataLayersSet.Empty();
+						Mode->SelectedDataLayerActors.Empty();
 
 						FDataLayerCreationParameters CreationParams;
 						CreationParams.DataLayerAsset = const_cast<UDataLayerAsset*>(DataLayerAsset);
 						CreationParams.ParentDataLayer = InParentDataLayer;
-						CreationParams.WorlDataLayers = InParentDataLayer ? InParentDataLayer->GetOuterAWorldDataLayers() : GetOwningWorldAWorldDataLayers();
-						if (UDataLayerInstance* NewDataLayerInstance = DataLayerEditorSubsystem->CreateDataLayerInstance(CreationParams))
+						CreationParams.WorlDataLayers = InParentDataLayer ? InParentDataLayer->GetOuterAWorldDataLayers() : Mode->GetOwningWorldAWorldDataLayers();
+						if (UDataLayerInstance* NewDataLayerInstance = UDataLayerEditorSubsystem::Get()->CreateDataLayerInstance(CreationParams))
 						{
-							SelectedDataLayersSet.Add(NewDataLayerInstance);
+							Mode->SelectedDataLayersSet.Add(NewDataLayerInstance);
 							// Select it and open a rename when it gets refreshed
 							SceneOutliner->OnItemAdded(NewDataLayerInstance, SceneOutliner::ENewItemAction::Select | SceneOutliner::ENewItemAction::Rename);
 						}
@@ -1041,11 +1041,11 @@ void FDataLayerMode::RegisterContextMenu()
 				};
 
 				// If selection contains readonly DataLayerInstances
-				if (Algo::AnyOf(SelectedDataLayersSet, [](const TWeakObjectPtr<const UDataLayerInstance>& DataLayerInstance) { return DataLayerInstance.IsValid() && DataLayerInstance->IsReadOnly(); }))
+				if (Algo::AnyOf(Mode->SelectedDataLayersSet, [](const TWeakObjectPtr<const UDataLayerInstance>& DataLayerInstance) { return DataLayerInstance.IsValid() && DataLayerInstance->IsReadOnly(); }))
 				{
 					// Create a menu entry that allows to creating DataLayerInstance in the world using selection referenced DataLayer Assets
 					TSet<const UDataLayerAsset*> ExistingDataLayerAssets;
-					if (AWorldDataLayers* WorldDataLayer = GetOwningWorldAWorldDataLayers())
+					if (AWorldDataLayers* WorldDataLayer = Mode->GetOwningWorldAWorldDataLayers())
 					{
 						WorldDataLayer->ForEachDataLayer([&ExistingDataLayerAssets](UDataLayerInstance* DataLayerInstance)
 						{
@@ -1059,7 +1059,7 @@ void FDataLayerMode::RegisterContextMenu()
 					}
 
 					TSet<const UDataLayerAsset*> DataLayerAssetsToImport;
-					for (const TWeakObjectPtr<const UDataLayerInstance>& DataLayerInstance : SelectedDataLayersSet)
+					for (const TWeakObjectPtr<const UDataLayerInstance>& DataLayerInstance : Mode->SelectedDataLayersSet)
 					{
 						if (DataLayerInstance.IsValid() && DataLayerInstance->IsReadOnly())
 						{
@@ -1075,7 +1075,7 @@ void FDataLayerMode::RegisterContextMenu()
 					if (!DataLayerAssetsToImport.IsEmpty())
 					{
 						Section.AddMenuEntry("ImportDataLayers", LOCTEXT("ImportDataLayers", "Import Data Layer(s)"), FText(), FSlateIcon(),
-							FUIAction(FExecuteAction::CreateLambda([this, DataLayerAssetsToImport, CreateNewDataLayer]()
+							FUIAction(FExecuteAction::CreateLambda([DataLayerAssetsToImport, CreateNewDataLayer]()
 								{
 									const FScopedTransaction Transaction(LOCTEXT("ImportDataLayersTransaction", "Import Data Layers"));
 									CreateNewDataLayer(nullptr, DataLayerAssetsToImport);
@@ -1087,27 +1087,27 @@ void FDataLayerMode::RegisterContextMenu()
 				}
 				
 				Section.AddMenuEntry("CreateNewDataLayer", LOCTEXT("CreateNewDataLayer", "Create New Data Layer"), FText(), FSlateIcon(),
-					FUIAction(FExecuteAction::CreateLambda([this, CreateNewDataLayer]() { CreateNewDataLayer(); })));
+					FUIAction(FExecuteAction::CreateLambda([CreateNewDataLayer]() { CreateNewDataLayer(); })));
 
-				UDataLayerInstance* ParentDataLayer = (SelectedDataLayersSet.Num() == 1) ? const_cast<UDataLayerInstance*>(SelectedDataLayersSet.CreateIterator()->Get()) : nullptr;
+				UDataLayerInstance* ParentDataLayer = (Mode->SelectedDataLayersSet.Num() == 1) ? const_cast<UDataLayerInstance*>(Mode->SelectedDataLayersSet.CreateIterator()->Get()) : nullptr;
 				if (ParentDataLayer)
 				{
 					Section.AddMenuEntry("CreateNewDataLayerUnderDataLayer", FText::Format(LOCTEXT("CreateNewDataLayerUnderDataLayer", "Create New Data Layer under \"{0}\""), FText::FromString(ParentDataLayer->GetDataLayerShortName())), FText(), FSlateIcon(),
-					FUIAction(FExecuteAction::CreateLambda([this, CreateNewDataLayer, ParentDataLayer]() { CreateNewDataLayer(ParentDataLayer); })));
+					FUIAction(FExecuteAction::CreateLambda([CreateNewDataLayer, ParentDataLayer]() { CreateNewDataLayer(ParentDataLayer); })));
 				}
 
 				Section.AddMenuEntry("AddSelectedActorsToNewDataLayer", LOCTEXT("AddSelectedActorsToNewDataLayer", "Add Selected Actors to New Data Layer"), FText(), FSlateIcon(),
 					FUIAction(
-						FExecuteAction::CreateLambda([this]() {
-							if (UDataLayerAsset* DataLayerAsset = PromptDataLayerAssetSelection())
+						FExecuteAction::CreateLambda([Mode]() {
+							if (UDataLayerAsset* DataLayerAsset = Mode->PromptDataLayerAssetSelection())
 							{
 								const FScopedTransaction Transaction(LOCTEXT("AddSelectedActorsToNewDataLayer", "Add Selected Actors to New Data Layer"));
 								FDataLayerCreationParameters CreationParams;
 								CreationParams.DataLayerAsset = DataLayerAsset;
-								CreationParams.WorlDataLayers = GetOwningWorldAWorldDataLayers();
-								if (UDataLayerInstance* NewDataLayerInstance = DataLayerEditorSubsystem->CreateDataLayerInstance(CreationParams))
+								CreationParams.WorlDataLayers = Mode->GetOwningWorldAWorldDataLayers();
+								if (UDataLayerInstance* NewDataLayerInstance = UDataLayerEditorSubsystem::Get()->CreateDataLayerInstance(CreationParams))
 								{
-									DataLayerEditorSubsystem->AddSelectedActorsToDataLayer(NewDataLayerInstance);
+									UDataLayerEditorSubsystem::Get()->AddSelectedActorsToDataLayer(NewDataLayerInstance);
 								}
 							}}),
 						FCanExecuteAction::CreateLambda([] { return GEditor->GetSelectedActorCount() > 0; })
@@ -1115,29 +1115,29 @@ void FDataLayerMode::RegisterContextMenu()
 
 				Section.AddMenuEntry("AddSelectedActorsToSelectedDataLayers", LOCTEXT("AddSelectedActorsToSelectedDataLayersMenu", "Add Selected Actors to Selected Data Layers"), FText(), FSlateIcon(),
 					FUIAction(
-						FExecuteAction::CreateLambda([this,SelectedDataLayers]() {
+						FExecuteAction::CreateLambda([SelectedDataLayers]() {
 							check(!SelectedDataLayers.IsEmpty());
 							{
 								const FScopedTransaction Transaction(LOCTEXT("AddSelectedActorsToSelectedDataLayers", "Add Selected Actors to Selected Data Layers"));
-								DataLayerEditorSubsystem->AddSelectedActorsToDataLayers(SelectedDataLayers);
+								UDataLayerEditorSubsystem::Get()->AddSelectedActorsToDataLayers(SelectedDataLayers);
 							}}),
 						FCanExecuteAction::CreateLambda([SelectedDataLayers,bSelectedDataLayersContainsLocked] { return !SelectedDataLayers.IsEmpty() && GEditor->GetSelectedActorCount() > 0 && !bSelectedDataLayersContainsLocked; })
 					));
 
-				if (!SelectedDataLayerActors.IsEmpty())
+				if (!Mode->SelectedDataLayerActors.IsEmpty())
 				{
 					Section.AddSubMenu("AddSelectedActorsTo", LOCTEXT("AddSelectedActorsTo", "Add Selected Actors To"), FText(),
-						FNewToolMenuDelegate::CreateLambda([this](UToolMenu* InSubMenu)
+						FNewToolMenuDelegate::CreateLambda([Mode](UToolMenu* InSubMenu)
 						{
-							CreateDataLayerPicker(InSubMenu, FOnDataLayerPicked::CreateLambda([this](UDataLayerInstance* TargetDataLayer)
+							CreateDataLayerPicker(InSubMenu, FOnDataLayerPicked::CreateLambda([Mode](UDataLayerInstance* TargetDataLayer)
 							{
 								check(TargetDataLayer);
 								TArray<AActor*> Actors;
-								Algo::TransformIf(SelectedDataLayerActors, Actors, [](const FSelectedDataLayerActor& InActor) { return InActor.Value.Get(); }, [](const FSelectedDataLayerActor& InActor) { return const_cast<AActor*>(InActor.Value.Get()); });
+								Algo::TransformIf(Mode->SelectedDataLayerActors, Actors, [](const FSelectedDataLayerActor& InActor) { return InActor.Value.Get(); }, [](const FSelectedDataLayerActor& InActor) { return const_cast<AActor*>(InActor.Value.Get()); });
 								if (!Actors.IsEmpty())
 								{
 									const FScopedTransaction Transaction(LOCTEXT("AddSelectedActorsToDataLayer", "Add Selected Actors to Selected Data Layer"));
-									DataLayerEditorSubsystem->AddActorsToDataLayers(Actors, { TargetDataLayer });
+									UDataLayerEditorSubsystem::Get()->AddActorsToDataLayers(Actors, { TargetDataLayer });
 								}
 							}));
 						}));
@@ -1145,19 +1145,19 @@ void FDataLayerMode::RegisterContextMenu()
 				if (!SelectedDataLayers.IsEmpty() && !bSelectedDataLayersContainsLocked)
 				{
 					Section.AddSubMenu("MoveSelectedDataLayersTo", LOCTEXT("MoveSelectedDataLayersTo", "Move Data Layers To"), FText(),
-						FNewToolMenuDelegate::CreateLambda([this](UToolMenu* InSubMenu)
+						FNewToolMenuDelegate::CreateLambda([Mode](UToolMenu* InSubMenu)
 						{
-							CreateDataLayerPicker(InSubMenu, FOnDataLayerPicked::CreateLambda([this](UDataLayerInstance* TargetDataLayer)
+							CreateDataLayerPicker(InSubMenu, FOnDataLayerPicked::CreateLambda([Mode](UDataLayerInstance* TargetDataLayer)
 							{
 								TArray<UDataLayerInstance*> DataLayers;
-								for (TWeakObjectPtr<const UDataLayerInstance>& DataLayer : SelectedDataLayersSet)
+								for (TWeakObjectPtr<const UDataLayerInstance>& DataLayer : Mode->SelectedDataLayersSet)
 								{
 									if (DataLayer.IsValid() && !DataLayer->IsLocked() && DataLayer != TargetDataLayer)
 									{
 										DataLayers.Add(const_cast<UDataLayerInstance*>(DataLayer.Get()));
 									}
 								}
-								SetParentDataLayer(DataLayers, TargetDataLayer);
+								Mode->SetParentDataLayer(DataLayers, TargetDataLayer);
 							}), /*bShowRoot*/true);
 						}));
 				}
@@ -1166,33 +1166,33 @@ void FDataLayerMode::RegisterContextMenu()
 
 				Section.AddMenuEntry("RemoveSelectedActorsFromSelectedDataLayers", LOCTEXT("RemoveSelectedActorsFromSelectedDataLayersMenu", "Remove Selected Actors from Selected Data Layers"), FText(), FSlateIcon(),
 					FUIAction(
-						FExecuteAction::CreateLambda([this,SelectedDataLayers]() {
+						FExecuteAction::CreateLambda([SelectedDataLayers]() {
 							check(!SelectedDataLayers.IsEmpty());
 							{
 								const FScopedTransaction Transaction(LOCTEXT("RemoveSelectedActorsFromSelectedDataLayers_DataLayerMode", "Remove Selected Actors from Selected Data Layers"));
-								DataLayerEditorSubsystem->RemoveSelectedActorsFromDataLayers(SelectedDataLayers);
+								UDataLayerEditorSubsystem::Get()->RemoveSelectedActorsFromDataLayers(SelectedDataLayers);
 							}}),
 						FCanExecuteAction::CreateLambda([SelectedDataLayers,bSelectedDataLayersContainsLocked] { return !SelectedDataLayers.IsEmpty() && GEditor->GetSelectedActorCount() > 0 && !bSelectedDataLayersContainsLocked; })
 					));
 
 				Section.AddMenuEntry("DeleteSelectedDataLayers", LOCTEXT("DeleteSelectedDataLayers", "Delete Selected Data Layers"), FText(), FSlateIcon(),
 					FUIAction(
-						FExecuteAction::CreateLambda([this,SelectedDataLayers]() {
+						FExecuteAction::CreateLambda([SelectedDataLayers]() {
 							check(!SelectedDataLayers.IsEmpty());
 							{
 								const FScopedTransaction Transaction(LOCTEXT("DeleteSelectedDataLayers", "Delete Selected Data Layers"));
-								DataLayerEditorSubsystem->DeleteDataLayers(SelectedDataLayers);
+								UDataLayerEditorSubsystem::Get()->DeleteDataLayers(SelectedDataLayers);
 							}}),
 						FCanExecuteAction::CreateLambda([SelectedDataLayers,bSelectedDataLayersContainsLocked] { return !SelectedDataLayers.IsEmpty() && !bSelectedDataLayersContainsLocked; })
 					));
 
 				Section.AddMenuEntry("RenameSelectedDataLayer", LOCTEXT("RenameSelectedDataLayer", "Rename Selected Data Layer"), FText(), FSlateIcon(),
 					FUIAction(
-						FExecuteAction::CreateLambda([this,SceneOutliner,SelectedDataLayers]() {
+						FExecuteAction::CreateLambda([Mode,SceneOutliner,SelectedDataLayers]() {
 							if (SelectedDataLayers.Num() == 1)
 							{
 								FSceneOutlinerTreeItemPtr ItemToRename = SceneOutliner->GetTreeItem(SelectedDataLayers[0]);
-								if (ItemToRename.IsValid() && CanRenameItem(*ItemToRename) && ItemToRename->CanInteract())
+								if (ItemToRename.IsValid() && Mode->CanRenameItem(*ItemToRename) && ItemToRename->CanInteract())
 								{
 									SceneOutliner->SetPendingRenameItem(ItemToRename);
 									SceneOutliner->ScrollItemIntoView(ItemToRename);
@@ -1209,34 +1209,34 @@ void FDataLayerMode::RegisterContextMenu()
 
 				Section.AddMenuEntry("SelectActorsInDataLayers", LOCTEXT("SelectActorsInDataLayers", "Select Actors in Data Layers"), FText(), FSlateIcon(),
 					FUIAction(
-						FExecuteAction::CreateLambda([this,SelectedDataLayers]() {
+						FExecuteAction::CreateLambda([SelectedDataLayers]() {
 							check(!SelectedDataLayers.IsEmpty());
 							{
 								const FScopedTransaction Transaction(LOCTEXT("SelectActorsInDataLayers", "Select Actors in Data Layers"));
 								GEditor->SelectNone(/*bNoteSelectionChange*/false, /*bDeselectBSPSurfs*/true);
-								DataLayerEditorSubsystem->SelectActorsInDataLayers(SelectedDataLayers, /*bSelect*/true, /*bNotify*/true, /*bSelectEvenIfHidden*/true);
+								UDataLayerEditorSubsystem::Get()->SelectActorsInDataLayers(SelectedDataLayers, /*bSelect*/true, /*bNotify*/true, /*bSelectEvenIfHidden*/true);
 							}}),
 						FCanExecuteAction::CreateLambda([SelectedDataLayers] { return !SelectedDataLayers.IsEmpty(); })
 					));
 
 				Section.AddMenuEntry("AppendActorsToSelection", LOCTEXT("AppendActorsToSelection", "Append Actors in Data Layer to Selection"), FText(), FSlateIcon(),
 					FUIAction(
-						FExecuteAction::CreateLambda([this,SelectedDataLayers]() {
+						FExecuteAction::CreateLambda([SelectedDataLayers]() {
 							check(!SelectedDataLayers.IsEmpty());
 							{
 								const FScopedTransaction Transaction(LOCTEXT("AppendActorsToSelection", "Append Actors in Data Layer to Selection"));
-								DataLayerEditorSubsystem->SelectActorsInDataLayers(SelectedDataLayers, /*bSelect*/true, /*bNotify*/true, /*bSelectEvenIfHidden*/true);
+								UDataLayerEditorSubsystem::Get()->SelectActorsInDataLayers(SelectedDataLayers, /*bSelect*/true, /*bNotify*/true, /*bSelectEvenIfHidden*/true);
 							}}),
 						FCanExecuteAction::CreateLambda([SelectedDataLayers,bSelectedDataLayersContainsLocked] { return !SelectedDataLayers.IsEmpty() && !bSelectedDataLayersContainsLocked; })
 					));
 
 				Section.AddMenuEntry("DeselectActors", LOCTEXT("DeselectActors", "Deselect Actors in Data Layer"), FText(), FSlateIcon(),
 					FUIAction(
-						FExecuteAction::CreateLambda([this,SelectedDataLayers]() {
+						FExecuteAction::CreateLambda([SelectedDataLayers]() {
 							check(!SelectedDataLayers.IsEmpty());
 							{
 								const FScopedTransaction Transaction(LOCTEXT("DeselectActors", "Deselect Actors in Data Layer"));
-								DataLayerEditorSubsystem->SelectActorsInDataLayers(SelectedDataLayers, /*bSelect*/false, /*bNotifySelectActors*/true);
+								UDataLayerEditorSubsystem::Get()->SelectActorsInDataLayers(SelectedDataLayers, /*bSelect*/false, /*bNotifySelectActors*/true);
 							}}),
 						FCanExecuteAction::CreateLambda([SelectedDataLayers] { return !SelectedDataLayers.IsEmpty(); })
 					));
@@ -1273,46 +1273,46 @@ void FDataLayerMode::RegisterContextMenu()
 
 				Section.AddMenuEntry("MakeCurrentDataLayers", LOCTEXT("MakeCurrentDataLayers", "Make Current Data Layer(s)"), FText(), FSlateIcon(),
 					FUIAction(
-						FExecuteAction::CreateLambda([this, SceneOutliner, SelectedDataLayers]() 
+						FExecuteAction::CreateLambda([Mode, SceneOutliner, SelectedDataLayers]() 
 						{
 							const FScopedTransaction Transaction(LOCTEXT("MakeCurrentDataLayers", "Make Current Data Layer(s)"));
-							for (TWeakObjectPtr<const UDataLayerInstance>& DataLayer : SelectedDataLayersSet)
+							for (TWeakObjectPtr<const UDataLayerInstance>& DataLayer : Mode->SelectedDataLayersSet)
 							{
 								if (DataLayer.IsValid() && !DataLayer->IsLocked())
 								{
-									DataLayerEditorSubsystem->AddToActorEditorContext(const_cast<UDataLayerInstance*>(DataLayer.Get()));
+									UDataLayerEditorSubsystem::Get()->AddToActorEditorContext(const_cast<UDataLayerInstance*>(DataLayer.Get()));
 								}
 							}
 						}),
-						FCanExecuteAction::CreateLambda([this] { return Algo::AnyOf(SelectedDataLayersSet, [](const TWeakObjectPtr<const UDataLayerInstance>& DataLayer) { return DataLayer.IsValid() && !DataLayer->IsLocked() && !DataLayer->IsInActorEditorContext(); }); })
+						FCanExecuteAction::CreateLambda([Mode] { return Algo::AnyOf(Mode->SelectedDataLayersSet, [](const TWeakObjectPtr<const UDataLayerInstance>& DataLayer) { return DataLayer.IsValid() && !DataLayer->IsLocked() && !DataLayer->IsInActorEditorContext(); }); })
 					));
 
 				Section.AddMenuEntry("RemoveCurrentDataLayers", LOCTEXT("RemoveCurrentDataLayers", "Remove Current Data Layer(s)"), FText(), FSlateIcon(),
 					FUIAction(
-						FExecuteAction::CreateLambda([this, SceneOutliner, SelectedDataLayers]()
+						FExecuteAction::CreateLambda([Mode, SceneOutliner, SelectedDataLayers]()
 						{
 							const FScopedTransaction Transaction(LOCTEXT("RemoveCurrentDataLayers", "Remove Current Data Layer(s)"));
-							for (TWeakObjectPtr<const UDataLayerInstance>& DataLayer : SelectedDataLayersSet)
+							for (TWeakObjectPtr<const UDataLayerInstance>& DataLayer : Mode->SelectedDataLayersSet)
 							{
 								if (DataLayer.IsValid() && !DataLayer->IsLocked())
 								{
-									DataLayerEditorSubsystem->RemoveFromActorEditorContext(const_cast<UDataLayerInstance*>(DataLayer.Get()));
+									UDataLayerEditorSubsystem::Get()->RemoveFromActorEditorContext(const_cast<UDataLayerInstance*>(DataLayer.Get()));
 								}
 							}
 						}),
-						FCanExecuteAction::CreateLambda([this] { return Algo::AnyOf(SelectedDataLayersSet, [](const TWeakObjectPtr<const UDataLayerInstance>& DataLayer) { return DataLayer.IsValid() && !DataLayer->IsLocked() && DataLayer->IsInActorEditorContext(); }); })
+						FCanExecuteAction::CreateLambda([Mode] { return Algo::AnyOf(Mode->SelectedDataLayersSet, [](const TWeakObjectPtr<const UDataLayerInstance>& DataLayer) { return DataLayer.IsValid() && !DataLayer->IsLocked() && DataLayer->IsInActorEditorContext(); }); })
 					));
 
 				Section.AddMenuEntry("ClearCurrentDataLayers", LOCTEXT("ClearCurrentDataLayers", "Clear Current Data Layers"), FText(), FSlateIcon(),
 					FUIAction(
-						FExecuteAction::CreateLambda([this, AllDataLayers]() 
+						FExecuteAction::CreateLambda([AllDataLayers]() 
 						{
 							check(!AllDataLayers.IsEmpty());
 							{
 								const FScopedTransaction Transaction(LOCTEXT("ClearCurrentDataLayers", "Clear Current Data Layers"));
 								for (const UDataLayerInstance* DataLayer : AllDataLayers)
 								{
-									DataLayerEditorSubsystem->RemoveFromActorEditorContext(const_cast<UDataLayerInstance*>(DataLayer));
+									UDataLayerEditorSubsystem::Get()->RemoveFromActorEditorContext(const_cast<UDataLayerInstance*>(DataLayer));
 								}
 							}}),
 						FCanExecuteAction::CreateLambda([bHasActorEditorContextDataLayers] { return bHasActorEditorContextDataLayers; })
@@ -1324,11 +1324,11 @@ void FDataLayerMode::RegisterContextMenu()
 
 				Section.AddMenuEntry("MakeAllDataLayersVisible", LOCTEXT("MakeAllDataLayersVisible", "Make All Data Layers Visible"), FText(), FSlateIcon(),
 					FUIAction(
-						FExecuteAction::CreateLambda([this,AllDataLayers]() {
+						FExecuteAction::CreateLambda([AllDataLayers]() {
 							check(!AllDataLayers.IsEmpty());
 							{
 								const FScopedTransaction Transaction(LOCTEXT("MakeAllDataLayersVisible", "Make All Data Layers Visible"));
-								DataLayerEditorSubsystem->MakeAllDataLayersVisible();
+								UDataLayerEditorSubsystem::Get()->MakeAllDataLayersVisible();
 							}}),
 						FCanExecuteAction::CreateLambda([AllDataLayers] { return !AllDataLayers.IsEmpty(); })
 					));
@@ -1342,13 +1342,6 @@ void FDataLayerMode::RegisterContextMenu()
 	}
 }
 
-void FDataLayerMode::UnregisterContextMenu()
-{
-	UToolMenus* ToolMenus = UToolMenus::Get();
-	ToolMenus->RemoveMenu(DefaultContextBaseMenuName);
-	ToolMenus->RemoveMenu(DefaultContextMenuName);
-}
-
 TSharedPtr<SWidget> FDataLayerMode::CreateContextMenu()
 {
 	RegisterContextMenu();
@@ -1694,7 +1687,7 @@ TSharedRef<FSceneOutlinerFilter> FDataLayerMode::CreateShowOnlySelectedActorsFil
 	return MakeShareable(new FDataLayerActorFilter(FDataLayerActorTreeItem::FFilterPredicate::CreateStatic(IsActorSelected), FSceneOutlinerFilter::EDefaultBehaviour::Pass, FDataLayerActorTreeItem::FFilterPredicate::CreateStatic(IsActorSelected)));
 }
 
-UDataLayerAsset* FDataLayerMode::PromptDataLayerAssetSelection()
+UDataLayerAsset* FDataLayerMode::PromptDataLayerAssetSelection() const
 {
 	IContentBrowserSingleton& ContentBrowserSingleton = FModuleManager::LoadModuleChecked<FContentBrowserModule>("ContentBrowser").Get();
 
diff --git a/Engine/Source/Editor/DataLayerEditor/Private/DataLayer/DataLayerMode.h b/Engine/Source/Editor/DataLayerEditor/Private/DataLayer/DataLayerMode.h
--- a/Engine/Source/Editor/DataLayerEditor/Private/DataLayer/DataLayerMode.h
+++ b/Engine/Source/Editor/DataLayerEditor/Private/DataLayer/DataLayerMode.h
@@ -101,8 +101,7 @@ protected:
 
 private:
 	/* Private Helpers */
-	void RegisterContextMenu();
-	void UnregisterContextMenu();
+	static void RegisterContextMenu();
 	void ChooseRepresentingWorld();
 	void OnSelectWorld(TWeakObjectPtr<UWorld> World);
 	bool IsWorldChecked(TWeakObjectPtr<UWorld> World) const;
@@ -120,7 +119,7 @@ private:
 	AWorldDataLayers* GetOwningWorldAWorldDataLayers() const;
 	FSceneOutlinerDragValidationInfo ValidateDrop(const ISceneOutlinerTreeItem& DropTarget, bool bMoveOperation = false) const;
 
-	UDataLayerAsset* PromptDataLayerAssetSelection();
+	UDataLayerAsset* PromptDataLayerAssetSelection() const;
 
 	/** Filter factories */
 	static TSharedRef<FSceneOutlinerFilter> CreateShowOnlySelectedActorsFilter();
