Remove prompt to choose Data Layer asset when world has deprecated Data Layers.
Remove option to import Data Layer assets from Level Instance if world has deprecated Data Layers.

diff --git a/Engine/Source/Editor/DataLayerEditor/Private/DataLayer/DataLayerMode.cpp b/Engine/Source/Editor/DataLayerEditor/Private/DataLayer/DataLayerMode.cpp
--- a/Engine/Source/Editor/DataLayerEditor/Private/DataLayer/DataLayerMode.cpp
+++ b/Engine/Source/Editor/DataLayerEditor/Private/DataLayer/DataLayerMode.cpp
@@ -1008,34 +1008,47 @@ void FDataLayerMode::RegisterContextMenu()
 			{
 				FToolMenuSection& Section = InMenu->AddSection("DataLayers", LOCTEXT("DataLayers", "Data Layers"));
 				
-				auto CreateNewDataLayer = [SceneOutliner, Mode](UDataLayerInstance* InParentDataLayer = nullptr, const TSet<const UDataLayerAsset*>& InDataLayerAssets = {})
+				auto CreateNewDataLayerInternal = [SceneOutliner, Mode](UDataLayerInstance* InParentDataLayer, const UDataLayerAsset* InDataLayerAsset = nullptr)
 				{
-					TSet<const UDataLayerAsset*> DataLayerAssets = InDataLayerAssets;
-					DataLayerAssets.Remove(nullptr);
+					check(InDataLayerAsset || UDataLayerEditorSubsystem::Get()->HasDeprecatedDataLayers());
 
-					if (DataLayerAssets.IsEmpty())
+					const FScopedTransaction Transaction(LOCTEXT("CreateNewDataLayer", "Create New Data Layer"));
+					Mode->SelectedDataLayersSet.Empty();
+					Mode->SelectedDataLayerActors.Empty();
+
+					FDataLayerCreationParameters CreationParams;
+					CreationParams.DataLayerAsset = const_cast<UDataLayerAsset*>(InDataLayerAsset);
+					CreationParams.ParentDataLayer = InParentDataLayer;
+					CreationParams.WorlDataLayers = InParentDataLayer ? InParentDataLayer->GetOuterAWorldDataLayers() : Mode->GetOwningWorldAWorldDataLayers();
+					if (UDataLayerInstance* NewDataLayerInstance = UDataLayerEditorSubsystem::Get()->CreateDataLayerInstance(CreationParams))
 					{
-						if (UDataLayerAsset* DataLayerAsset = Mode->PromptDataLayerAssetSelection())
-						{
-							DataLayerAssets.Add(DataLayerAsset);
-						}
+						Mode->SelectedDataLayersSet.Add(NewDataLayerInstance);
+						// Select it and open a rename when it gets refreshed
+						SceneOutliner->OnItemAdded(NewDataLayerInstance, SceneOutliner::ENewItemAction::Select | SceneOutliner::ENewItemAction::Rename);
 					}
+				};
 
-					for (const UDataLayerAsset* DataLayerAsset : DataLayerAssets)
+				auto CreateNewDataLayer = [SceneOutliner, Mode, CreateNewDataLayerInternal](UDataLayerInstance* InParentDataLayer = nullptr, const TSet<const UDataLayerAsset*>& InDataLayerAssets = {})
+				{
+					if (UDataLayerEditorSubsystem::Get()->HasDeprecatedDataLayers())
+					{
+						check(InDataLayerAssets.IsEmpty());
+						CreateNewDataLayerInternal(InParentDataLayer);
+					}
+					else
 					{
-						const FScopedTransaction Transaction(LOCTEXT("CreateNewDataLayer", "Create New Data Layer"));
-						Mode->SelectedDataLayersSet.Empty();
-						Mode->SelectedDataLayerActors.Empty();
-
-						FDataLayerCreationParameters CreationParams;
-						CreationParams.DataLayerAsset = const_cast<UDataLayerAsset*>(DataLayerAsset);
-						CreationParams.ParentDataLayer = InParentDataLayer;
-						CreationParams.WorlDataLayers = InParentDataLayer ? InParentDataLayer->GetOuterAWorldDataLayers() : Mode->GetOwningWorldAWorldDataLayers();
-						if (UDataLayerInstance* NewDataLayerInstance = UDataLayerEditorSubsystem::Get()->CreateDataLayerInstance(CreationParams))
+						TSet<const UDataLayerAsset*> DataLayerAssets = InDataLayerAssets;
+						DataLayerAssets.Remove(nullptr);
+						if (DataLayerAssets.IsEmpty())
+						{
+							if (UDataLayerAsset* DataLayerAsset = Mode->PromptDataLayerAssetSelection())
+							{
+								DataLayerAssets.Add(DataLayerAsset);
+							}
+						}
+						for (const UDataLayerAsset* DataLayerAsset : DataLayerAssets)
 						{
-							Mode->SelectedDataLayersSet.Add(NewDataLayerInstance);
-							// Select it and open a rename when it gets refreshed
-							SceneOutliner->OnItemAdded(NewDataLayerInstance, SceneOutliner::ENewItemAction::Select | SceneOutliner::ENewItemAction::Rename);
+							CreateNewDataLayerInternal(InParentDataLayer, DataLayerAsset);
 						}
 					}
 				};
@@ -1072,7 +1085,7 @@ void FDataLayerMode::RegisterContextMenu()
 						}
 					}
 
-					if (!DataLayerAssetsToImport.IsEmpty())
+					if (!DataLayerAssetsToImport.IsEmpty() && !UDataLayerEditorSubsystem::Get()->HasDeprecatedDataLayers())
 					{
 						Section.AddMenuEntry("ImportDataLayers", LOCTEXT("ImportDataLayers", "Import Data Layer(s)"), FText(), FSlateIcon(),
 							FUIAction(FExecuteAction::CreateLambda([DataLayerAssetsToImport, CreateNewDataLayer]()
