Data Layers: fix crash when destructing the data layers editor subsystem.

#rnx

diff --git a/Engine/Source/Editor/DataLayerEditor/Private/DataLayer/DataLayerEditorSubsystem.cpp b/Engine/Source/Editor/DataLayerEditor/Private/DataLayer/DataLayerEditorSubsystem.cpp
--- a/Engine/Source/Editor/DataLayerEditor/Private/DataLayer/DataLayerEditorSubsystem.cpp
+++ b/Engine/Source/Editor/DataLayerEditor/Private/DataLayer/DataLayerEditorSubsystem.cpp
@@ -169,12 +169,27 @@ void UDataLayerEditorSubsystem::Deinitialize()
 
 	Super::Deinitialize();
 
-	DataLayersBroadcast->Deinitialize();
+	if (DataLayersBroadcast)
+	{
+		DataLayersBroadcast->Deinitialize();
+		DataLayersBroadcast.Reset();
+	}
 
 	// Unregister the engine broadcast bridge
 	DataLayerEditorLoadingStateChanged.Remove(OnActorDataLayersEditorLoadingStateChangedEngineBridgeHandle);
 }
 
+void UDataLayerEditorSubsystem::BeginDestroy()
+{
+	if (DataLayersBroadcast)
+	{
+		DataLayersBroadcast->Deinitialize();
+		DataLayersBroadcast.Reset();
+	}
+
+	Super::BeginDestroy();
+}
+
 void UDataLayerEditorSubsystem::OnExecuteActorEditorContextAction(UWorld* InWorld, const EActorEditorContextAction& InType, AActor* InActor)
 {
 	UE_CLOG(!InWorld, LogDataLayerEditorSubsystem, Error, TEXT("%s - Failed because world in null."), ANSI_TO_TCHAR(__FUNCTION__));
diff --git a/Engine/Source/Editor/DataLayerEditor/Public/DataLayer/DataLayerEditorSubsystem.h b/Engine/Source/Editor/DataLayerEditor/Public/DataLayer/DataLayerEditorSubsystem.h
--- a/Engine/Source/Editor/DataLayerEditor/Public/DataLayer/DataLayerEditorSubsystem.h
+++ b/Engine/Source/Editor/DataLayerEditor/Public/DataLayer/DataLayerEditorSubsystem.h
@@ -66,6 +66,10 @@ public:
 	 */
 	virtual ~UDataLayerEditorSubsystem() {}
 
+	// UObject overrides.
+	virtual void BeginDestroy() override;
+	// ~END UObject overrides.
+
 	//~ Begin IActorEditorContextClient interface
 	virtual void OnExecuteActorEditorContextAction(UWorld* InWorld, const EActorEditorContextAction& InType, AActor* InActor = nullptr) override;
 	virtual bool GetActorEditorContextDisplayInfo(UWorld* InWorld, FActorEditorContextClientDisplayInfo& OutDiplayInfo) const override;
