Fix crash on editor exit

diff --git a/Engine/Source/Editor/DataLayerEditor/Private/DataLayer/DataLayerEditorSubsystem.cpp b/Engine/Source/Editor/DataLayerEditor/Private/DataLayer/DataLayerEditorSubsystem.cpp
--- a/Engine/Source/Editor/DataLayerEditor/Private/DataLayer/DataLayerEditorSubsystem.cpp
+++ b/Engine/Source/Editor/DataLayerEditor/Private/DataLayer/DataLayerEditorSubsystem.cpp
@@ -81,7 +81,10 @@ void FDataLayersBroadcast::Deinitialize()
 		FEditorDelegates::MapChange.RemoveAll(this);
 		FEditorDelegates::PostUndoRedo.RemoveAll(this);
 		FCoreUObjectDelegates::OnObjectPropertyChanged.RemoveAll(this);
-		GEngine->OnLevelActorAdded().RemoveAll(this);
+		if (GEngine)
+		{
+			GEngine->OnLevelActorAdded().RemoveAll(this);
+		}
 		USelection::SelectionChangedEvent.RemoveAll(this);
 		USelection::SelectObjectEvent.RemoveAll(this);
 	}
