Added detail customization for DataLayerInstance to hide Runtime parameters when it's not a Runtime Data Layer.

diff --git a/Engine/Source/Editor/DataLayerEditor/Private/DataLayer/DataLayerInstanceCustomization.cpp b/Engine/Source/Editor/DataLayerEditor/Private/DataLayer/DataLayerInstanceCustomization.cpp
new file mode 100644
--- /dev/null
+++ b/Engine/Source/Editor/DataLayerEditor/Private/DataLayer/DataLayerInstanceCustomization.cpp
@@ -0,0 +1,34 @@
+// Copyright Epic Games, Inc. All Rights Reserved.
+
+#include "DataLayerInstanceCustomization.h"
+#include "DetailLayoutBuilder.h"
+
+#define LOCTEXT_NAMESPACE "FDataLayerInstanceDetails"
+
+TSharedRef<IDetailCustomization> FDataLayerInstanceDetails::MakeInstance()
+{
+	return MakeShareable(new FDataLayerInstanceDetails);
+}
+
+void FDataLayerInstanceDetails::CustomizeDetails(IDetailLayoutBuilder& DetailBuilder)
+{
+	TArray<TWeakObjectPtr<UObject>> ObjectsBeingCustomized;
+	DetailBuilder.GetObjectsBeingCustomized(ObjectsBeingCustomized);
+	
+	bool bHasRuntime = false;
+	for (const TWeakObjectPtr<UObject>& SelectedObject : ObjectsBeingCustomized)
+	{
+		UDataLayerInstance* DataLayerInstance = Cast<UDataLayerInstance>(SelectedObject.Get());
+		if (DataLayerInstance && DataLayerInstance->IsRuntime())
+		{
+			bHasRuntime = true;
+			break;
+		}
+	}
+	if (!bHasRuntime)
+	{
+		DetailBuilder.HideProperty(GET_MEMBER_NAME_CHECKED(UDataLayerInstance, InitialRuntimeState));
+	}
+}
+
+#undef LOCTEXT_NAMESPACE
diff --git a/Engine/Source/Editor/DataLayerEditor/Private/DataLayer/DataLayerInstanceCustomization.h b/Engine/Source/Editor/DataLayerEditor/Private/DataLayer/DataLayerInstanceCustomization.h
new file mode 100644
--- /dev/null
+++ b/Engine/Source/Editor/DataLayerEditor/Private/DataLayer/DataLayerInstanceCustomization.h
@@ -0,0 +1,26 @@
+// Copyright Epic Games, Inc. All Rights Reserved.
+
+#pragma once
+
+#include "CoreMinimal.h"
+#include "SlateFwd.h"
+#include "IDetailCustomization.h"
+#include "WorldPartition/DataLayer/DataLayerInstance.h"
+
+class IDetailLayoutBuilder;
+
+enum class ECheckBoxState : uint8;
+
+class FDataLayerInstanceDetails : public IDetailCustomization
+{
+public:
+	/** Makes a new instance of this detail layout class for a specific detail view requesting it */
+	static TSharedRef<IDetailCustomization> MakeInstance();
+
+protected:
+	FDataLayerInstanceDetails() {}
+
+private:
+	/** IDetailCustomization interface */
+	virtual void CustomizeDetails(IDetailLayoutBuilder& DetailBuilder) override;
+};
\ No newline at end of file
diff --git a/Engine/Source/Editor/DataLayerEditor/Private/DataLayerEditorModule.cpp b/Engine/Source/Editor/DataLayerEditor/Private/DataLayerEditorModule.cpp
--- a/Engine/Source/Editor/DataLayerEditor/Private/DataLayerEditorModule.cpp
+++ b/Engine/Source/Editor/DataLayerEditor/Private/DataLayerEditorModule.cpp
@@ -9,6 +9,7 @@
 #include "Editor.h"
 #include "Modules/ModuleManager.h"
 #include "DataLayer/DataLayerPropertyTypeCustomization.h"
+#include "DataLayer/DataLayerInstanceCustomization.h"
 #include "DataLayer/SDataLayerBrowser.h"
 #include "DataLayer/DataLayerNameEditSink.h"
 #include "DataLayer/DataLayerEditorSubsystem.h"
@@ -17,12 +18,13 @@
 IMPLEMENT_MODULE(FDataLayerEditorModule, DataLayerEditor );
 
 static const FName NAME_ActorDataLayer(TEXT("ActorDataLayer"));
-
+static const FName NAME_DataLayerInstance(TEXT("DataLayerInstance"));
 
 void FDataLayerEditorModule::StartupModule()
 {
 	FPropertyEditorModule& PropertyModule = FModuleManager::LoadModuleChecked<FPropertyEditorModule>("PropertyEditor");
 	PropertyModule.RegisterCustomPropertyTypeLayout(NAME_ActorDataLayer, FOnGetPropertyTypeCustomizationInstance::CreateLambda([] { return MakeShared<FDataLayerPropertyTypeCustomization>(); }));
+	PropertyModule.RegisterCustomClassLayout(NAME_DataLayerInstance, FOnGetDetailCustomizationInstance::CreateStatic(&FDataLayerInstanceDetails::MakeInstance));
 
 	FEditorWidgetsModule& EditorWidgetsModule = FModuleManager::LoadModuleChecked<FEditorWidgetsModule>("EditorWidgets");
 	EditorWidgetsModule.GetObjectNameEditSinkRegistry()->RegisterObjectNameEditSink(MakeShared<FDataLayerNameEditSink>());
@@ -33,6 +35,7 @@ void FDataLayerEditorModule::ShutdownModule()
 	if (FPropertyEditorModule* PropertyModule = FModuleManager::GetModulePtr<FPropertyEditorModule>("PropertyEditor"))
 	{
 		PropertyModule->UnregisterCustomPropertyTypeLayout(NAME_ActorDataLayer);
+		PropertyModule->UnregisterCustomPropertyTypeLayout(NAME_DataLayerInstance);
 	}
 }
 
diff --git a/Engine/Source/Runtime/Engine/Private/WorldPartition/DataLayer/DataLayerInstance.cpp b/Engine/Source/Runtime/Engine/Private/WorldPartition/DataLayer/DataLayerInstance.cpp
--- a/Engine/Source/Runtime/Engine/Private/WorldPartition/DataLayer/DataLayerInstance.cpp
+++ b/Engine/Source/Runtime/Engine/Private/WorldPartition/DataLayer/DataLayerInstance.cpp
@@ -96,6 +96,11 @@ bool UDataLayerInstance::CanEditChange(const FProperty* InProperty) const
 		return false;
 	}
 
+	if (!IsRuntime() && (InProperty->GetFName() == GET_MEMBER_NAME_CHECKED(UDataLayerInstance, InitialRuntimeState)))
+	{
+		return false;
+	}
+
 	return true;
 }
 
diff --git a/Engine/Source/Runtime/Engine/Public/WorldPartition/DataLayer/DataLayerInstance.h b/Engine/Source/Runtime/Engine/Public/WorldPartition/DataLayer/DataLayerInstance.h
--- a/Engine/Source/Runtime/Engine/Public/WorldPartition/DataLayer/DataLayerInstance.h
+++ b/Engine/Source/Runtime/Engine/Public/WorldPartition/DataLayer/DataLayerInstance.h
@@ -40,6 +40,7 @@ class ENGINE_API UDataLayerInstance : public UObject
 	GENERATED_UCLASS_BODY()
 
 	friend class UDataLayerConversionInfo;
+	friend class FDataLayerInstanceDetails;
 
 public:
 	virtual void PostLoad() override;
@@ -157,7 +158,7 @@ protected:
 	uint32 bIsLocked : 1;
 #endif
 
-	UPROPERTY(Category = "Data Layer|Advanced|Runtime", EditAnywhere, meta = (EditConditionHides, EditCondition = "bIsRuntime"))
+	UPROPERTY(Category = "Data Layer|Advanced|Runtime", EditAnywhere)
 	EDataLayerRuntimeState InitialRuntimeState;
 
 private:
