Optimized editor when modifying actors with Data Layers : DataLayer outliner and actor visibility updates are deferred at next DataLayerEditorSubsystem::Tick.

diff --git a/Engine/Source/Editor/DataLayerEditor/Private/DataLayer/DataLayerEditorSubsystem.cpp b/Engine/Source/Editor/DataLayerEditor/Private/DataLayer/DataLayerEditorSubsystem.cpp
--- a/Engine/Source/Editor/DataLayerEditor/Private/DataLayer/DataLayerEditorSubsystem.cpp
+++ b/Engine/Source/Editor/DataLayerEditor/Private/DataLayer/DataLayerEditorSubsystem.cpp
@@ -142,6 +142,8 @@ void FDataLayersBroadcast::OnObjectPostEditChange(UObject* Object, FPropertyChan
 
 UDataLayerEditorSubsystem::UDataLayerEditorSubsystem()
 : bRebuildSelectedDataLayersFromEditorSelection(false)
+, bAsyncBroadcastDataLayerChanged(false)
+, bAsyncUpdateAllActorsVisibility(false)
 {}
 
 UDataLayerEditorSubsystem* UDataLayerEditorSubsystem::Get()
@@ -198,6 +200,37 @@ void UDataLayerEditorSubsystem::BeginDestroy()
 	Super::BeginDestroy();
 }
 
+
+UWorld* UDataLayerEditorSubsystem::GetTickableGameObjectWorld() const
+{
+	return GetWorld();
+}
+
+ETickableTickType UDataLayerEditorSubsystem::GetTickableTickType() const
+{
+	return IsTemplate() ? ETickableTickType::Never : ETickableTickType::Conditional;
+}
+
+bool UDataLayerEditorSubsystem::IsAllowedToTick() const
+{
+	return GetWorld() && (bAsyncBroadcastDataLayerChanged || bAsyncUpdateAllActorsVisibility);
+}
+
+void UDataLayerEditorSubsystem::Tick(float DeltaTime)
+{
+	if (bAsyncBroadcastDataLayerChanged)
+	{
+		BroadcastDataLayerChanged(EDataLayerAction::Reset, NULL, NAME_None);
+		bAsyncBroadcastDataLayerChanged = false;
+	}
+
+	if (bAsyncUpdateAllActorsVisibility)
+	{
+		UpdateAllActorsVisibility(false, false);
+		bAsyncUpdateAllActorsVisibility = false;
+	}
+}
+
 void UDataLayerEditorSubsystem::OnExecuteActorEditorContextAction(UWorld* InWorld, const EActorEditorContextAction& InType, AActor* InActor)
 {
 	UE_CLOG(!InWorld, LogDataLayerEditorSubsystem, Error, TEXT("%s - Failed because world in null."), ANSI_TO_TCHAR(__FUNCTION__));
@@ -320,8 +353,8 @@ void UDataLayerEditorSubsystem::EditorMapChange()
 
 void UDataLayerEditorSubsystem::EditorRefreshDataLayerBrowser()
 {
-	BroadcastDataLayerChanged(EDataLayerAction::Reset, NULL, NAME_None);
-	UpdateAllActorsVisibility(false, false);
+	bAsyncBroadcastDataLayerChanged = true;
+	bAsyncUpdateAllActorsVisibility = true;
 }
 
 void UDataLayerEditorSubsystem::PostUndoRedo()
diff --git a/Engine/Source/Editor/DataLayerEditor/Public/DataLayer/DataLayerEditorSubsystem.h b/Engine/Source/Editor/DataLayerEditor/Public/DataLayer/DataLayerEditorSubsystem.h
--- a/Engine/Source/Editor/DataLayerEditor/Public/DataLayer/DataLayerEditorSubsystem.h
+++ b/Engine/Source/Editor/DataLayerEditor/Public/DataLayer/DataLayerEditorSubsystem.h
@@ -37,7 +37,7 @@ struct FDataLayerCreationParameters
 };
 
 UCLASS()
-class DATALAYEREDITOR_API UDataLayerEditorSubsystem final : public UEditorSubsystem, public IActorEditorContextClient
+class DATALAYEREDITOR_API UDataLayerEditorSubsystem final : public UEditorSubsystem, public IActorEditorContextClient, public FTickableGameObject
 {
 	GENERATED_BODY()
 
@@ -63,6 +63,15 @@ public:
 	 */
 	virtual ~UDataLayerEditorSubsystem() {}
 
+	//~ Begin FTickableGameObject interface
+	virtual UWorld* GetTickableGameObjectWorld() const override;
+	virtual bool IsTickableInEditor() const { return true; }
+	virtual ETickableTickType GetTickableTickType() const override;
+	virtual bool IsAllowedToTick() const override;
+	virtual void Tick(float DeltaTime) override;
+	TStatId GetStatId() const override { RETURN_QUICK_DECLARE_CYCLE_STAT(UDataLayerEditorSubsystem, STATGROUP_Tickables); }
+	//~ End FTickableGameObject interface
+
 	// UObject overrides.
 	virtual void BeginDestroy() override;
 	// ~END UObject overrides.
@@ -679,6 +688,12 @@ private:
 	/** Internal flag to know if SelectedDataLayersFromEditorSelection needs to be rebuilt. */
 	mutable bool bRebuildSelectedDataLayersFromEditorSelection;
 
+	/** When true, next Tick will call BroadcastDataLayerChanged */
+	bool bAsyncBroadcastDataLayerChanged;
+
+	/** When true, next Tick will call UpdateAllActorsVisibility */
+	bool bAsyncUpdateAllActorsVisibility;
+
 	/** Fires whenever one or more DataLayer changes */
 	FOnDataLayerChanged DataLayerChanged;
 
