Fix crash when exiting editor in FDataLayersBroadcast::Deinitialize.

#rnx

diff --git a/Engine/Source/Editor/DataLayerEditor/Private/DataLayer/DataLayerEditorSubsystem.cpp b/Engine/Source/Editor/DataLayerEditor/Private/DataLayer/DataLayerEditorSubsystem.cpp
--- a/Engine/Source/Editor/DataLayerEditor/Private/DataLayer/DataLayerEditorSubsystem.cpp
+++ b/Engine/Source/Editor/DataLayerEditor/Private/DataLayer/DataLayerEditorSubsystem.cpp
@@ -79,17 +79,21 @@ void FDataLayersBroadcast::Deinitialize()
 	if (bIsInitialized)
 	{
 		bIsInitialized = false;
-		FEditorDelegates::MapChange.RemoveAll(this);
-		FEditorDelegates::PostUndoRedo.RemoveAll(this);
-		FCoreUObjectDelegates::OnObjectPropertyChanged.RemoveAll(this);
-		if (GEngine)
+
+		if (!IsEngineExitRequested())
 		{
-			GEngine->OnLevelActorAdded().RemoveAll(this);
+			FEditorDelegates::MapChange.RemoveAll(this);
+			FEditorDelegates::PostUndoRedo.RemoveAll(this);
+			FCoreUObjectDelegates::OnObjectPropertyChanged.RemoveAll(this);
+			if (GEngine)
+			{
+				GEngine->OnLevelActorAdded().RemoveAll(this);
+			}
+			USelection::SelectionChangedEvent.RemoveAll(this);
+			USelection::SelectObjectEvent.RemoveAll(this);
+			UDataLayerSubsystem::OnWorldDataLayerPostRegister.RemoveAll(this);
+			UDataLayerSubsystem::OnWorldDataLayerPreUnregister.RemoveAll(this);
 		}
-		USelection::SelectionChangedEvent.RemoveAll(this);
-		USelection::SelectObjectEvent.RemoveAll(this);
-		UDataLayerSubsystem::OnWorldDataLayerPostRegister.RemoveAll(this);
-		UDataLayerSubsystem::OnWorldDataLayerPreUnregister.RemoveAll(this);
 	}
 }
 
