Fix crash occuring when adding actors to a deprecated datalayer

diff --git a/Engine/Source/Editor/DataLayerEditor/Private/DataLayer/DataLayerEditorSubsystem.cpp b/Engine/Source/Editor/DataLayerEditor/Private/DataLayer/DataLayerEditorSubsystem.cpp
--- a/Engine/Source/Editor/DataLayerEditor/Private/DataLayer/DataLayerEditorSubsystem.cpp
+++ b/Engine/Source/Editor/DataLayerEditor/Private/DataLayer/DataLayerEditorSubsystem.cpp
@@ -474,7 +474,7 @@ bool UDataLayerEditorSubsystem::AddActorsToDataLayers(const TArray<AActor*>& Act
 			bool bActorWasModified = false;
 			for (const UDataLayerInstance* DataLayerInstance : DataLayers)
 			{
-				if (const UDataLayerInstanceWithAsset* DataLayerInstanceWithAsset = CastChecked<UDataLayerInstanceWithAsset>(DataLayerInstance))
+				if (const UDataLayerInstanceWithAsset* DataLayerInstanceWithAsset = Cast<UDataLayerInstanceWithAsset>(DataLayerInstance))
 				{
 					// If actor's level doesn't match this DataLayerInstance outer level, 
 					// Make sure that a DataLayer Instance for this Data Layer Asset exists in the Actor's level.
