Data Layers: fix crash when destructing the data layers editor subsystem.

#rnx

diff --git a/Engine/Source/Editor/DataLayerEditor/Private/DataLayer/DataLayerEditorSubsystem.cpp b/Engine/Source/Editor/DataLayerEditor/Private/DataLayer/DataLayerEditorSubsystem.cpp
--- a/Engine/Source/Editor/DataLayerEditor/Private/DataLayer/DataLayerEditorSubsystem.cpp
+++ b/Engine/Source/Editor/DataLayerEditor/Private/DataLayer/DataLayerEditorSubsystem.cpp
@@ -183,27 +183,12 @@ void UDataLayerEditorSubsystem::Deinitialize()
 
 	Super::Deinitialize();
 
-	if (DataLayersBroadcast)
-	{
-		DataLayersBroadcast->Deinitialize();
-		DataLayersBroadcast.Reset();
-	}
+	DataLayersBroadcast->Deinitialize();
 
 	// Unregister the engine broadcast bridge
 	DataLayerEditorLoadingStateChanged.Remove(OnActorDataLayersEditorLoadingStateChangedEngineBridgeHandle);
 }
 
-void UDataLayerEditorSubsystem::BeginDestroy()
-{
-	if (DataLayersBroadcast)
-	{
-		DataLayersBroadcast->Deinitialize();
-		DataLayersBroadcast.Reset();
-	}
-
-	Super::BeginDestroy();
-}
-
 
 UWorld* UDataLayerEditorSubsystem::GetTickableGameObjectWorld() const
 {
@@ -235,6 +220,17 @@ void UDataLayerEditorSubsystem::Tick(float DeltaTime)
 	}
 }
 
+void UDataLayerEditorSubsystem::BeginDestroy()
+{
+	if (DataLayersBroadcast)
+	{
+		DataLayersBroadcast->Deinitialize();
+		DataLayersBroadcast.Reset();
+	}
+
+	Super::BeginDestroy();
+}
+
 void UDataLayerEditorSubsystem::OnExecuteActorEditorContextAction(UWorld* InWorld, const EActorEditorContextAction& InType, AActor* InActor)
 {
 	UE_CLOG(!InWorld, LogDataLayerEditorSubsystem, Error, TEXT("%s - Failed because world in null."), ANSI_TO_TCHAR(__FUNCTION__));
diff --git a/Engine/Source/Editor/DataLayerEditor/Public/DataLayer/DataLayerEditorSubsystem.h b/Engine/Source/Editor/DataLayerEditor/Public/DataLayer/DataLayerEditorSubsystem.h
--- a/Engine/Source/Editor/DataLayerEditor/Public/DataLayer/DataLayerEditorSubsystem.h
+++ b/Engine/Source/Editor/DataLayerEditor/Public/DataLayer/DataLayerEditorSubsystem.h
@@ -72,9 +72,9 @@ public:
 	TStatId GetStatId() const override { RETURN_QUICK_DECLARE_CYCLE_STAT(UDataLayerEditorSubsystem, STATGROUP_Tickables); }
 	//~ End FTickableGameObject interface
 
-	// UObject overrides.
+	//~ Begin UObject overrides.
 	virtual void BeginDestroy() override;
-	// ~END UObject overrides.
+	// ~End UObject overrides.
 
 	//~ Begin IActorEditorContextClient interface
 	virtual void OnExecuteActorEditorContextAction(UWorld* InWorld, const EActorEditorContextAction& InType, AActor* InActor = nullptr) override;
