Display WorldDataLayers Tree Items for world data layers which are not the main world partitions but are empty.
Allows to add DataLayerInstances by right-click on WorldDataLayers when they are not a "main worldpartiton" world data layer
Fix IsSubWorldDataLayers working in pie/cooked build for non-"main world partitition data layers"

diff --git a/Engine/Source/Editor/DataLayerEditor/Private/DataLayer/DataLayerHierarchy.cpp b/Engine/Source/Editor/DataLayerEditor/Private/DataLayer/DataLayerHierarchy.cpp
--- a/Engine/Source/Editor/DataLayerEditor/Private/DataLayer/DataLayerHierarchy.cpp
+++ b/Engine/Source/Editor/DataLayerEditor/Private/DataLayer/DataLayerHierarchy.cpp
@@ -181,6 +181,20 @@ void FDataLayerHierarchy::CreateItems(TArray<FSceneOutlinerTreeItemPtr>& OutItem
 		CreateDataLayerTreeItems(CurrentLevel);
 	}
 
+	// Create Tree items for WorldDataLayers(WDL) which are part of the current world (not subworld) but are not the ULevel::AWorldDataLayers.
+	// The only way to create data layer instances for in these WDL is to right-click on the WDL tree item and select create new.
+	// We cannot rely on creating the WDL tree item only if it has DataLayerInstance (::FindOrCreateParentItem), as it prevents to initially create a new instance.
+	DataLayerSubsystem->ForEachWorldDataLayer([this, &OutItems](AWorldDataLayers* WorldDataLayers)
+	{
+		if (!WorldDataLayers->IsSubWorldDataLayers() && !WorldDataLayers->IsTheMainWorldDataLayers())
+		{
+			FSceneOutlinerTreeItemPtr WorldDataLayerItem = Mode->CreateItemFor<FWorldDataLayersTreeItem>(WorldDataLayers, true);
+			OutItems.Add(WorldDataLayerItem);
+		}
+
+		return true;
+	});
+
 	if (bShowDataLayerActors)
 	{
 		for (AActor* Actor : FActorRange(GetOwningWorld()))
diff --git a/Engine/Source/Runtime/Engine/Private/WorldPartition/DataLayer/WorldDataLayers.cpp b/Engine/Source/Runtime/Engine/Private/WorldPartition/DataLayer/WorldDataLayers.cpp
--- a/Engine/Source/Runtime/Engine/Private/WorldPartition/DataLayer/WorldDataLayers.cpp
+++ b/Engine/Source/Runtime/Engine/Private/WorldPartition/DataLayer/WorldDataLayers.cpp
@@ -854,7 +854,14 @@ bool AWorldDataLayers::IsRuntimeRelevant() const
 bool AWorldDataLayers::IsSubWorldDataLayers() const
 {
 	check(GetWorld());
-	return GetTypedOuter<UWorld>() != GetWorld();
+	UWorld* OuterWorld = GetTypedOuter<UWorld>();
+	return OuterWorld != nullptr && OuterWorld->GetFName() != GetWorld()->GetFName();
+}
+
+bool AWorldDataLayers::IsTheMainWorldDataLayers() const
+{
+	UWorld* OuterWorld = GetTypedOuter<UWorld>();
+	return OuterWorld && OuterWorld->GetWorldDataLayers() == this;
 }
 
 #if WITH_EDITOR
diff --git a/Engine/Source/Runtime/Engine/Public/WorldPartition/DataLayer/WorldDataLayers.h b/Engine/Source/Runtime/Engine/Public/WorldPartition/DataLayer/WorldDataLayers.h
--- a/Engine/Source/Runtime/Engine/Public/WorldPartition/DataLayer/WorldDataLayers.h
+++ b/Engine/Source/Runtime/Engine/Public/WorldPartition/DataLayer/WorldDataLayers.h
@@ -89,6 +89,7 @@ public:
 #endif
 
 	bool IsSubWorldDataLayers() const;
+	bool IsTheMainWorldDataLayers() const;
 	bool IsRuntimeRelevant() const;
 	static FName GetWorldPartionWorldDataLayersName() { return FName(TEXT("WorldDataLayers")); } // reserved for ULevel::WorldDataLayers
 	
