Remove modal view that prompt a user to select a data layer asset
Data Layer Instance are now created empty
Asset can be created and assigned via the detail view

diff --git a/Engine/Source/Editor/DataLayerEditor/Private/DataLayer/DataLayerMode.cpp b/Engine/Source/Editor/DataLayerEditor/Private/DataLayer/DataLayerMode.cpp
--- a/Engine/Source/Editor/DataLayerEditor/Private/DataLayer/DataLayerMode.cpp
+++ b/Engine/Source/Editor/DataLayerEditor/Private/DataLayer/DataLayerMode.cpp
@@ -1080,10 +1080,8 @@ void FDataLayerMode::RegisterContextMenu()
 			{
 				FToolMenuSection& Section = InMenu->AddSection("DataLayers", LOCTEXT("DataLayers", "Data Layers"));
 				
-				auto CreateNewDataLayerInternal = [SceneOutliner, Mode](UDataLayerInstance* InParentDataLayer, const UDataLayerAsset* InDataLayerAsset = nullptr)
+				auto CreateNewDataLayerInternal = [SceneOutliner, Mode](UDataLayerInstance* InParentDataLayer, const UDataLayerAsset* InDataLayerAsset = nullptr) -> UDataLayerInstance*
 				{
-					check(InDataLayerAsset || UDataLayerEditorSubsystem::Get()->HasDeprecatedDataLayers());
-
 					const FScopedTransaction Transaction(LOCTEXT("CreateNewDataLayer", "Create New Data Layer"));
 					Mode->SelectedDataLayersSet.Empty();
 					Mode->SelectedDataLayerActors.Empty();
@@ -1100,34 +1098,31 @@ void FDataLayerMode::RegisterContextMenu()
 						Mode->SelectedDataLayersSet.Add(NewDataLayerInstance);
 						// Select it and open a rename when it gets refreshed
 						SceneOutliner->OnItemAdded(NewDataLayerInstance, SceneOutliner::ENewItemAction::Select | SceneOutliner::ENewItemAction::Rename);
+
+						return NewDataLayerInstance;
 					}
+
+					return nullptr;
 				};
 
-				auto CreateNewDataLayer = [SceneOutliner, Mode, CreateNewDataLayerInternal](UDataLayerInstance* InParentDataLayer = nullptr, const TSet<const UDataLayerAsset*>& InDataLayerAssets = {})
+				auto CreateNewDataLayer = [SceneOutliner, Mode, CreateNewDataLayerInternal](UDataLayerInstance* InParentDataLayer = nullptr, const UDataLayerAsset* InDataLayerAsset = nullptr) -> UDataLayerInstance*
 				{
 					if (UDataLayerEditorSubsystem::Get()->HasDeprecatedDataLayers())
 					{
-						check(InDataLayerAssets.IsEmpty());
-						CreateNewDataLayerInternal(InParentDataLayer);
+						check(InDataLayerAsset == nullptr);
+						return CreateNewDataLayerInternal(InParentDataLayer);
 					}
 					else
 					{
-						TSet<const UDataLayerAsset*> DataLayerAssets = InDataLayerAssets;
-						DataLayerAssets.Remove(nullptr);
-						if (DataLayerAssets.IsEmpty())
-						{
-							if (UDataLayerAsset* DataLayerAsset = Mode->PromptDataLayerAssetSelection())
-							{
-								DataLayerAssets.Add(DataLayerAsset);
-							}
-						}
-						for (const UDataLayerAsset* DataLayerAsset : DataLayerAssets)
-						{
-							CreateNewDataLayerInternal(InParentDataLayer, DataLayerAsset);
-						}
+						return CreateNewDataLayerInternal(InParentDataLayer, InDataLayerAsset);
 					}
 				};
 
+				auto CreateNewEmptyDataLayer = [SceneOutliner, Mode, CreateNewDataLayerInternal](UDataLayerInstance* InParentDataLayer = nullptr) -> UDataLayerInstance*
+				{
+					return CreateNewDataLayerInternal(InParentDataLayer);
+				};
+
 				// If selection contains readonly DataLayerInstances
 				if (Algo::AnyOf(Mode->SelectedDataLayersSet, [](const TWeakObjectPtr<const UDataLayerInstance>& DataLayerInstance) { return DataLayerInstance.IsValid() && DataLayerInstance->IsReadOnly(); }))
 				{
@@ -1166,7 +1161,10 @@ void FDataLayerMode::RegisterContextMenu()
 							FUIAction(FExecuteAction::CreateLambda([DataLayerAssetsToImport, CreateNewDataLayer]()
 								{
 									const FScopedTransaction Transaction(LOCTEXT("ImportDataLayersTransaction", "Import Data Layers"));
-									CreateNewDataLayer(nullptr, DataLayerAssetsToImport);
+									for (const UDataLayerAsset* DataLayerAsset : DataLayerAssetsToImport)
+									{
+										CreateNewDataLayer(nullptr, DataLayerAsset);
+									}
 								})));
 					}
 
@@ -1175,29 +1173,24 @@ void FDataLayerMode::RegisterContextMenu()
 				}
 				
 				Section.AddMenuEntry("CreateNewDataLayer", LOCTEXT("CreateNewDataLayer", "Create New Data Layer"), FText(), FSlateIcon(),
-					FUIAction(FExecuteAction::CreateLambda([CreateNewDataLayer]() { CreateNewDataLayer(); })));
+					FUIAction(FExecuteAction::CreateLambda([CreateNewEmptyDataLayer]() { CreateNewEmptyDataLayer(); })));
 
 				UDataLayerInstance* ParentDataLayer = (Mode->SelectedDataLayersSet.Num() == 1) ? const_cast<UDataLayerInstance*>(Mode->SelectedDataLayersSet.CreateIterator()->Get()) : nullptr;
 				if (ParentDataLayer)
 				{
 					Section.AddMenuEntry("CreateNewDataLayerUnderDataLayer", FText::Format(LOCTEXT("CreateNewDataLayerUnderDataLayer", "Create New Data Layer under \"{0}\""), FText::FromString(ParentDataLayer->GetDataLayerShortName())), FText(), FSlateIcon(),
-					FUIAction(FExecuteAction::CreateLambda([CreateNewDataLayer, ParentDataLayer]() { CreateNewDataLayer(ParentDataLayer); })));
+					FUIAction(FExecuteAction::CreateLambda([CreateNewEmptyDataLayer, ParentDataLayer]() { CreateNewEmptyDataLayer(ParentDataLayer); })));
 				}
 
 				Section.AddMenuEntry("AddSelectedActorsToNewDataLayer", LOCTEXT("AddSelectedActorsToNewDataLayer", "Add Selected Actors to New Data Layer"), FText(), FSlateIcon(),
 					FUIAction(
-						FExecuteAction::CreateLambda([Mode]() {
-							if (UDataLayerAsset* DataLayerAsset = Mode->PromptDataLayerAssetSelection())
-							{
+						FExecuteAction::CreateLambda([Mode, CreateNewEmptyDataLayer]() {
 								const FScopedTransaction Transaction(LOCTEXT("AddSelectedActorsToNewDataLayer", "Add Selected Actors to New Data Layer"));
-								FDataLayerCreationParameters CreationParams;
-								CreationParams.DataLayerAsset = DataLayerAsset;
-								CreationParams.WorldDataLayers = Mode->GetOwningWorldAWorldDataLayers();
-								if (UDataLayerInstance* NewDataLayerInstance = UDataLayerEditorSubsystem::Get()->CreateDataLayerInstance(CreationParams))
+								if (UDataLayerInstance* NewDataLayerInstance = CreateNewEmptyDataLayer())
 								{
 									UDataLayerEditorSubsystem::Get()->AddSelectedActorsToDataLayer(NewDataLayerInstance);
 								}
-							}}),
+							}),
 						FCanExecuteAction::CreateLambda([] { return GEditor->GetSelectedActorCount() > 0; })
 					));
 
@@ -1775,26 +1768,6 @@ TSharedRef<FSceneOutlinerFilter> FDataLayerMode::CreateShowOnlySelectedActorsFil
 	return MakeShareable(new FDataLayerActorFilter(FDataLayerActorTreeItem::FFilterPredicate::CreateStatic(IsActorSelected), FSceneOutlinerFilter::EDefaultBehaviour::Pass, FDataLayerActorTreeItem::FFilterPredicate::CreateStatic(IsActorSelected)));
 }
 
-UDataLayerAsset* FDataLayerMode::PromptDataLayerAssetSelection() const
-{
-	IContentBrowserSingleton& ContentBrowserSingleton = FModuleManager::LoadModuleChecked<FContentBrowserModule>("ContentBrowser").Get();
-
-	FOpenAssetDialogConfig Config;
-	Config.bAllowMultipleSelection = false;
-	Config.AssetClassNames.Add(UDataLayerAsset::StaticClass()->GetClassPathName());
-	Config.DefaultPath = PickDataLayerDialogPath;
-	Config.DialogTitleOverride = LOCTEXT("PickDataLayerAssetDialogTitle", "Pick a Data Layer Asset");
-
-	TArray<FAssetData> Assets = ContentBrowserSingleton.CreateModalOpenAssetDialog(Config);
-	if (Assets.Num() == 1)
-	{
-		Assets[0].PackagePath.ToString(PickDataLayerDialogPath);
-		return CastChecked<UDataLayerAsset>(Assets[0].GetAsset());
-	}
-
-	return nullptr;
-}
-
 void FDataLayerMode::SynchronizeSelection()
 {
 	if (!bShowOnlySelectedActors && !bHighlightSelectedDataLayers)
diff --git a/Engine/Source/Editor/DataLayerEditor/Private/DataLayer/SDataLayerTreeLabel.cpp b/Engine/Source/Editor/DataLayerEditor/Private/DataLayer/SDataLayerTreeLabel.cpp
--- a/Engine/Source/Editor/DataLayerEditor/Private/DataLayer/SDataLayerTreeLabel.cpp
+++ b/Engine/Source/Editor/DataLayerEditor/Private/DataLayer/SDataLayerTreeLabel.cpp
@@ -36,6 +36,7 @@
 #include "Widgets/Views/STableRow.h"
 #include "WorldPartition/DataLayer/DataLayerInstance.h"
 #include "WorldPartition/DataLayer/DataLayerSubsystem.h"
+#include "WorldPartition/DataLayer/DataLayerInstanceWithAsset.h"
 
 struct FSlateBrush;
 
@@ -136,6 +137,16 @@ bool SDataLayerTreeLabel::ShouldBeHighlighted() const
 	return DataLayerTreeItem && DataLayerTreeItem->ShouldBeHighlighted();
 }
 
+bool SDataLayerTreeLabel::ShouldBeItalic() const
+{
+	if (const UDataLayerInstanceWithAsset* DataLayerInstanceWithAsset = Cast<UDataLayerInstanceWithAsset>(DataLayerPtr.Get()))
+	{
+		return !DataLayerInstanceWithAsset->GetAsset();
+	}
+
+	return false;
+}
+
 bool SDataLayerTreeLabel::IsInActorEditorContext() const
 {
 	const UDataLayerInstance* DataLayer = DataLayerPtr.Get();
@@ -148,10 +159,12 @@ FSlateFontInfo SDataLayerTreeLabel::GetDisplayNameFont() const
 	{
 		return FAppStyle::Get().GetFontStyle("DataLayerBrowser.LabelFontBold");
 	}
-	else
+	else if(ShouldBeItalic())
 	{
-		return FAppStyle::Get().GetFontStyle("DataLayerBrowser.LabelFont");
+		return FAppStyle::Get().GetFontStyle("PropertyWindow.ItalicFont");
 	}
+
+	return FAppStyle::Get().GetFontStyle("DataLayerBrowser.LabelFont");
 }
 
 FText SDataLayerTreeLabel::GetDisplayText() const
@@ -172,8 +185,21 @@ FText SDataLayerTreeLabel::GetDisplayText() const
 		}
 	}
 
+	if (DataLayer)
+	{
+		if (const UDataLayerInstanceWithAsset* DataLayerInstanceWithAsset = Cast<UDataLayerInstanceWithAsset>(DataLayer))
+		{
+			if (!DataLayerInstanceWithAsset->GetAsset())
+			{
+				return LOCTEXT("DataLayerLabelUnknown", "Unknown");
+			}
+		}
+
+		return FText::Format(LOCTEXT("DataLayerDisplayText", "{0}{1}"), FText::FromString(DataLayer->GetDataLayerShortName()), SuffixText);
+	}
+	
 	static const FText DataLayerDeleted = LOCTEXT("DataLayerLabelForMissingDataLayer", "(Deleted Data Layer)");
-	return DataLayer ? FText::Format(LOCTEXT("DataLayerDisplayText", "{0}{1}"), FText::FromString(DataLayer->GetDataLayerShortName()), SuffixText) : DataLayerDeleted;
+	return DataLayerDeleted;
 }
 
 FText SDataLayerTreeLabel::GetTooltipText() const
diff --git a/Engine/Source/Editor/DataLayerEditor/Private/DataLayer/SDataLayerTreeLabel.h b/Engine/Source/Editor/DataLayerEditor/Private/DataLayer/SDataLayerTreeLabel.h
--- a/Engine/Source/Editor/DataLayerEditor/Private/DataLayer/SDataLayerTreeLabel.h
+++ b/Engine/Source/Editor/DataLayerEditor/Private/DataLayer/SDataLayerTreeLabel.h
@@ -42,6 +42,7 @@ private:
 	bool OnVerifyItemLabelChanged(const FText& InLabel, FText& OutErrorMessage);
 	void OnLabelCommitted(const FText& InLabel, ETextCommit::Type InCommitInfo);
 	bool ShouldBeHighlighted() const;
+	bool ShouldBeItalic() const;
 	bool IsInActorEditorContext() const;
 	void OnEnterEditingMode();
 	void OnExitEditingMode();
