Remove deleted function

diff --git a/Engine/Source/Editor/DataLayerEditor/Private/DataLayer/DataLayerMode.h b/Engine/Source/Editor/DataLayerEditor/Private/DataLayer/DataLayerMode.h
--- a/Engine/Source/Editor/DataLayerEditor/Private/DataLayer/DataLayerMode.h
+++ b/Engine/Source/Editor/DataLayerEditor/Private/DataLayer/DataLayerMode.h
@@ -150,8 +150,6 @@ private:
 	AWorldDataLayers* GetOwningWorldAWorldDataLayers() const;
 	FSceneOutlinerDragValidationInfo ValidateDrop(const ISceneOutlinerTreeItem& DropTarget, bool bMoveOperation = false) const;
 
-	UDataLayerAsset* PromptDataLayerAssetSelection() const;
-
 	/** Filter factories */
 	static TSharedRef<FSceneOutlinerFilter> CreateShowOnlySelectedActorsFilter();
 	static TSharedRef<FSceneOutlinerFilter> CreateHideEditorDataLayersFilter();
