Removed the menu option "Add Selected Actors to New Data Layer" in the Data Layer outliner as this doesn't work anymore with Data Layer assets.

diff --git a/Engine/Source/Editor/DataLayerEditor/Private/DataLayer/DataLayerMode.cpp b/Engine/Source/Editor/DataLayerEditor/Private/DataLayer/DataLayerMode.cpp
--- a/Engine/Source/Editor/DataLayerEditor/Private/DataLayer/DataLayerMode.cpp
+++ b/Engine/Source/Editor/DataLayerEditor/Private/DataLayer/DataLayerMode.cpp
@@ -1182,18 +1182,6 @@ void FDataLayerMode::RegisterContextMenu()
 					FUIAction(FExecuteAction::CreateLambda([CreateNewEmptyDataLayer, ParentDataLayer]() { CreateNewEmptyDataLayer(ParentDataLayer); })));
 				}
 
-				Section.AddMenuEntry("AddSelectedActorsToNewDataLayer", LOCTEXT("AddSelectedActorsToNewDataLayer", "Add Selected Actors to New Data Layer"), FText(), FSlateIcon(),
-					FUIAction(
-						FExecuteAction::CreateLambda([Mode, CreateNewEmptyDataLayer]() {
-								const FScopedTransaction Transaction(LOCTEXT("AddSelectedActorsToNewDataLayer", "Add Selected Actors to New Data Layer"));
-								if (UDataLayerInstance* NewDataLayerInstance = CreateNewEmptyDataLayer())
-								{
-									UDataLayerEditorSubsystem::Get()->AddSelectedActorsToDataLayer(NewDataLayerInstance);
-								}
-							}),
-						FCanExecuteAction::CreateLambda([] { return GEditor->GetSelectedActorCount() > 0; })
-					));
-
 				Section.AddMenuEntry("AddSelectedActorsToSelectedDataLayers", LOCTEXT("AddSelectedActorsToSelectedDataLayersMenu", "Add Selected Actors to Selected Data Layers"), FText(), FSlateIcon(),
 					FUIAction(
 						FExecuteAction::CreateLambda([SelectedDataLayers]() {
