Added UDataLayerEditorSubsystem::GetAllDataLayers()

diff --git a/Engine/Source/Editor/DataLayerEditor/Private/DataLayer/DataLayerEditorSubsystem.cpp b/Engine/Source/Editor/DataLayerEditor/Private/DataLayer/DataLayerEditorSubsystem.cpp
--- a/Engine/Source/Editor/DataLayerEditor/Private/DataLayer/DataLayerEditorSubsystem.cpp
+++ b/Engine/Source/Editor/DataLayerEditor/Private/DataLayer/DataLayerEditorSubsystem.cpp
@@ -1089,6 +1089,20 @@ bool UDataLayerEditorSubsystem::ToggleDataLayersIsLoadedInEditor(const TArray<UD
 	return true;
 }
 
+TArray<UDataLayerInstance*> UDataLayerEditorSubsystem::GetAllDataLayers()
+{
+	TArray<UDataLayerInstance*> DataLayerInstances;
+	if (UDataLayerSubsystem* DataLayerSubsystem = UWorld::GetSubsystem<UDataLayerSubsystem>(GetWorld()))
+	{
+		DataLayerSubsystem->ForEachDataLayer([&DataLayerInstances](UDataLayerInstance* DataLayerInstance)
+		{
+			DataLayerInstances.Add(DataLayerInstance);
+			return true;
+		});
+	}
+	return DataLayerInstances;
+}
+
 bool UDataLayerEditorSubsystem::ResetUserSettings()
 {
 	bool bChanged = false;
diff --git a/Engine/Source/Editor/DataLayerEditor/Public/DataLayer/DataLayerEditorSubsystem.h b/Engine/Source/Editor/DataLayerEditor/Public/DataLayer/DataLayerEditorSubsystem.h
--- a/Engine/Source/Editor/DataLayerEditor/Public/DataLayer/DataLayerEditorSubsystem.h
+++ b/Engine/Source/Editor/DataLayerEditor/Public/DataLayer/DataLayerEditorSubsystem.h
@@ -585,6 +585,12 @@ public:
 	UFUNCTION(BlueprintCallable, Category = DataLayers)
 	void DeleteDataLayer(UDataLayerInstance* DataLayerToDelete);
 
+	/**
+	 * Returns all Data Layers
+	 */
+	UFUNCTION(BlueprintCallable, Category = DataLayers)
+	TArray<UDataLayerInstance*> GetAllDataLayers();
+
 	/**
 	* Resets user override settings of all DataLayers
 	*/
