Fixed non-unity build caused by 22690448

diff --git a/Engine/Source/Editor/DataLayerEditor/Private/DataLayer/DataLayerOutlinerHasErrorColumn.cpp b/Engine/Source/Editor/DataLayerEditor/Private/DataLayer/DataLayerOutlinerHasErrorColumn.cpp
--- a/Engine/Source/Editor/DataLayerEditor/Private/DataLayer/DataLayerOutlinerHasErrorColumn.cpp
+++ b/Engine/Source/Editor/DataLayerEditor/Private/DataLayer/DataLayerOutlinerHasErrorColumn.cpp
@@ -19,6 +19,7 @@
 #include "Widgets/Images/SImage.h"
 #include "Widgets/SBoxPanel.h"
 #include "Widgets/SNullWidget.h"
+#include "Widgets/Text/STextBlock.h"
 #include "WorldPartition/DataLayer/DataLayerInstance.h"
 #include "WorldPartition/ErrorHandling/WorldPartitionStreamingGenerationTokenizedMessageErrorHandler.h"
 
