Data Layer Assets can now be dropped on DataLayerInstances and WorldDataLayers to create new instances

diff --git a/Engine/Source/Editor/DataLayerEditor/Private/DataLayer/DataLayerEditorSubsystem.cpp b/Engine/Source/Editor/DataLayerEditor/Private/DataLayer/DataLayerEditorSubsystem.cpp
--- a/Engine/Source/Editor/DataLayerEditor/Private/DataLayer/DataLayerEditorSubsystem.cpp
+++ b/Engine/Source/Editor/DataLayerEditor/Private/DataLayer/DataLayerEditorSubsystem.cpp
@@ -364,6 +364,16 @@ void UDataLayerEditorSubsystem::OnWorldDataLayersPreUnregister(AWorldDataLayers*
 	EditorRefreshDataLayerBrowser();
 }
 
+TArray<const UDataLayerInstance*> UDataLayerEditorSubsystem::GetDataLayerInstances(const TArray<const UDataLayerAsset*> DataLayerAssets) const
+{
+	if (UDataLayerSubsystem* DataLayerSubsystem = UWorld::GetSubsystem<UDataLayerSubsystem>(GetWorld()))
+	{
+		return DataLayerSubsystem->GetDataLayerInstances(DataLayerAssets);
+	}
+
+	return TArray<const UDataLayerInstance*>();
+}
+
 void UDataLayerEditorSubsystem::EditorMapChange()
 {
 	if (UWorld * World = GetWorld())
@@ -1151,6 +1161,18 @@ UDataLayerInstance* UDataLayerEditorSubsystem::GetDataLayerInstance(const UDataL
 	return nullptr;
 }
 
+TArray<UDataLayerInstance*> UDataLayerEditorSubsystem::GetDataLayerInstances(const TArray<UDataLayerAsset*> DataLayerAssets) const
+{
+	TArray<const UDataLayerAsset*> ConstAssets;
+	Algo::Transform(DataLayerAssets, ConstAssets, [](UDataLayerAsset* DataLayerAsset) { return DataLayerAsset; });
+
+	TArray<const UDataLayerInstance*> DataLayerInstances = GetDataLayerInstances(ConstAssets);
+
+	TArray<UDataLayerInstance*> Result;
+	Algo::Transform(DataLayerInstances, Result, [](const UDataLayerInstance* DataLayerInstance) { return const_cast<UDataLayerInstance*>(DataLayerInstance); });
+	return Result;
+}
+
 void UDataLayerEditorSubsystem::AddAllDataLayersTo(TArray<TWeakObjectPtr<UDataLayerInstance>>& OutDataLayers) const
 {
 	UE_CLOG(!GetWorld(), LogDataLayerEditorSubsystem, Error, TEXT("%s - Failed because world in null."), ANSI_TO_TCHAR(__FUNCTION__));
diff --git a/Engine/Source/Editor/DataLayerEditor/Private/DataLayer/DataLayerMode.cpp b/Engine/Source/Editor/DataLayerEditor/Private/DataLayer/DataLayerMode.cpp
--- a/Engine/Source/Editor/DataLayerEditor/Private/DataLayer/DataLayerMode.cpp
+++ b/Engine/Source/Editor/DataLayerEditor/Private/DataLayer/DataLayerMode.cpp
@@ -6,6 +6,10 @@
 #include "ActorMode.h"
 #include "Algo/AnyOf.h"
 #include "Algo/Transform.h"
+#include "Algo/Accumulate.h"
+#include "AssetRegistry/IAssetRegistry.h"
+#include "AssetRegistry/AssetRegistryModule.h"
+#include "AssetSelection.h"
 #include "AssetRegistry/AssetData.h"
 #include "Containers/EnumAsByte.h"
 #include "Containers/IndirectArray.h"
@@ -525,10 +529,13 @@ FReply FDataLayerMode::OnKeyDown(const FKeyEvent& InKeyEvent)
 
 bool FDataLayerMode::ParseDragDrop(FSceneOutlinerDragDropPayload& OutPayload, const FDragDropOperation& Operation) const
 {
-	return !GetDataLayerActorPairsFromOperation(Operation).IsEmpty() || !GetActorsFromOperation(Operation, true).IsEmpty() || !GetDataLayersFromOperation(Operation, true).IsEmpty();
+	return !GetDataLayerActorPairsFromOperation(Operation).IsEmpty()
+		|| !GetActorsFromOperation(Operation, true).IsEmpty()
+		|| !GetDataLayerInstancesFromOperation(Operation, true).IsEmpty()
+		|| !GetDataLayerAssetsFromOperation(Operation, true).IsEmpty();
 }
 
-FSceneOutlinerDragValidationInfo FDataLayerMode::ValidateDrop(const ISceneOutlinerTreeItem& DropTarget, bool bMoveOperation /*= false*/) const
+FSceneOutlinerDragValidationInfo FDataLayerMode::ValidateActorDrop(const ISceneOutlinerTreeItem& DropTarget, bool bMoveOperation /*= false*/) const
 {
 	if (const FDataLayerTreeItem* DataLayerItem = DropTarget.CastTo<FDataLayerTreeItem>())
 	{
@@ -575,86 +582,136 @@ FSceneOutlinerDragValidationInfo FDataLayerMode::ValidateDrop(const ISceneOutlin
 			}
 		}
 
-		return ValidateDrop(DropTarget);
+		return ValidateActorDrop(DropTarget);
 	}
-	else if (!GetDataLayerActorPairsFromOperation(Payload.SourceOperation).IsEmpty()) // Moving actor(s) into a Data Layer
+	
+	if (!GetDataLayerActorPairsFromOperation(Payload.SourceOperation).IsEmpty()) // Moving actor(s) into a Data Layer
+	{
+		return ValidateActorDrop(DropTarget, true);
+	}
+	 
+	TArray<const UDataLayerAsset*> DataLayerAssets = GetDataLayerAssetsFromOperation(Payload.SourceOperation);
+	if (!DataLayerAssets.IsEmpty())
 	{
-		return ValidateDrop(DropTarget, true);
+		return ValidateDataLayerAssetDrop(DropTarget, DataLayerAssets);
 	}
-	else // Moving a data layer(s)
+	
+	TArray<UDataLayerInstance*> PayloadDataLayers = GetDataLayerInstancesFromOperation(Payload.SourceOperation);
+	if (!PayloadDataLayers.IsEmpty())
 	{
-		TArray<UDataLayerInstance*> PayloadDataLayers = GetDataLayersFromOperation(Payload.SourceOperation);
-		if (!PayloadDataLayers.IsEmpty())
+		const FDataLayerTreeItem* DataLayerItem = DropTarget.CastTo<FDataLayerTreeItem>();
+		const FDataLayerActorTreeItem* DataLayerActorTreeItem = DropTarget.CastTo<FDataLayerActorTreeItem>();
+		const UDataLayerInstance* ParentDataLayer = DataLayerItem ? DataLayerItem->GetDataLayer() : nullptr;
+		if (!ParentDataLayer && DataLayerActorTreeItem)
+		{
+			ParentDataLayer = DataLayerActorTreeItem->GetDataLayer();
+		}
+
+		for (UDataLayerInstance* DataLayerInstance : PayloadDataLayers)
 		{
-			const FDataLayerTreeItem* DataLayerItem = DropTarget.CastTo<FDataLayerTreeItem>();
-			const FDataLayerActorTreeItem* DataLayerActorTreeItem = DropTarget.CastTo<FDataLayerActorTreeItem>();
-			const UDataLayerInstance* ParentDataLayer = DataLayerItem ? DataLayerItem->GetDataLayer() : nullptr;
-			if (!ParentDataLayer && DataLayerActorTreeItem)
+			if (DataLayerInstance->IsLocked())
 			{
-				ParentDataLayer = DataLayerActorTreeItem->GetDataLayer();
+				return FSceneOutlinerDragValidationInfo(ESceneOutlinerDropCompatibility::IncompatibleGeneric, LOCTEXT("CantMoveLockedDataLayer", "Can't move locked Data Layer"));
 			}
 
-			for (UDataLayerInstance* DataLayerInstance : PayloadDataLayers)
+			if (ParentDataLayer != nullptr)
 			{
-				if (DataLayerInstance->IsLocked())
-				{
-					return FSceneOutlinerDragValidationInfo(ESceneOutlinerDropCompatibility::IncompatibleGeneric, LOCTEXT("CantMoveLockedDataLayer", "Can't move locked Data Layer"));
-				}
-
-				if (ParentDataLayer != nullptr)
+				if (!DataLayerInstance->CanParent(ParentDataLayer))
 				{
-					if (!DataLayerInstance->CanParent(ParentDataLayer))
+					if (!DataLayerInstance->IsDataLayerTypeValidToParent(ParentDataLayer->GetType()))
 					{
-						if (!DataLayerInstance->IsDataLayerTypeValidToParent(ParentDataLayer->GetType()))
-						{
-							return FSceneOutlinerDragValidationInfo(ESceneOutlinerDropCompatibility::IncompatibleGeneric, 
-								FText::Format(LOCTEXT("CantMoveToDataLayerDiffType", "Can't move a {0} Data Layer under a {1} Data Layer"), UEnum::GetDisplayValueAsText(DataLayerInstance->GetType()), UEnum::GetDisplayValueAsText(ParentDataLayer->GetType())));
-						}
-
-						if (ParentDataLayer == DataLayerInstance->GetParent() || ParentDataLayer == DataLayerInstance)
-						{
-							return FSceneOutlinerDragValidationInfo(ESceneOutlinerDropCompatibility::IncompatibleGeneric, LOCTEXT("CantMoveToSameDataLayer", "Can't move Data Layer to same Data Layer"));
-						}
+						return FSceneOutlinerDragValidationInfo(ESceneOutlinerDropCompatibility::IncompatibleGeneric,
+							FText::Format(LOCTEXT("CantMoveToDataLayerDiffType", "Can't move a {0} Data Layer under a {1} Data Layer"), UEnum::GetDisplayValueAsText(DataLayerInstance->GetType()), UEnum::GetDisplayValueAsText(ParentDataLayer->GetType())));
+					}
 
-						if (ParentDataLayer->GetOuterAWorldDataLayers() != DataLayerInstance->GetOuterAWorldDataLayers())
-						{
-							return FSceneOutlinerDragValidationInfo(ESceneOutlinerDropCompatibility::IncompatibleGeneric, LOCTEXT("CantMoveToDifferentOuterWorldDataLayer", "Can't move Data Layer to a different World Data Layer"));
-						}
+					if (ParentDataLayer == DataLayerInstance->GetParent() || ParentDataLayer == DataLayerInstance)
+					{
+						return FSceneOutlinerDragValidationInfo(ESceneOutlinerDropCompatibility::IncompatibleGeneric, LOCTEXT("CantMoveToSameDataLayer", "Can't move Data Layer to same Data Layer"));
 					}
-				
-					if (ParentDataLayer->IsLocked())
+
+					if (ParentDataLayer->GetOuterAWorldDataLayers() != DataLayerInstance->GetOuterAWorldDataLayers())
 					{
-						return FSceneOutlinerDragValidationInfo(ESceneOutlinerDropCompatibility::IncompatibleGeneric, LOCTEXT("CantMoveDataLayerToLockedDataLayer", "Can't move Data Layer to locked Data Layer"));
+						return FSceneOutlinerDragValidationInfo(ESceneOutlinerDropCompatibility::IncompatibleGeneric, LOCTEXT("CantMoveToDifferentOuterWorldDataLayer", "Can't move Data Layer to a different World Data Layer"));
 					}
 				}
-			}
 
-			if (ParentDataLayer != nullptr)
-			{
-				return FSceneOutlinerDragValidationInfo(ESceneOutlinerDropCompatibility::Compatible, FText::Format(LOCTEXT("MoveDataLayerToDataLayer", "Move to Data Layer \"{0}\""), FText::FromString(ParentDataLayer->GetDataLayerShortName())));
-			}
-			else
-			{
-				return FSceneOutlinerDragValidationInfo(ESceneOutlinerDropCompatibility::Compatible, LOCTEXT("MoveDataLayerToRoot", "Move to root"));
+				if (ParentDataLayer->IsLocked())
+				{
+					return FSceneOutlinerDragValidationInfo(ESceneOutlinerDropCompatibility::IncompatibleGeneric, LOCTEXT("CantMoveDataLayerToLockedDataLayer", "Can't move Data Layer to locked Data Layer"));
+				}
 			}
 		}
+
+		if (ParentDataLayer != nullptr)
+		{
+			return FSceneOutlinerDragValidationInfo(ESceneOutlinerDropCompatibility::Compatible, FText::Format(LOCTEXT("MoveDataLayerToDataLayer", "Move to Data Layer \"{0}\""), FText::FromString(ParentDataLayer->GetDataLayerShortName())));
+		}
+		else
+		{
+			return FSceneOutlinerDragValidationInfo(ESceneOutlinerDropCompatibility::Compatible, LOCTEXT("MoveDataLayerToRoot", "Move to root"));
+		}	
 	}
 
 	return FSceneOutlinerDragValidationInfo::Invalid();
 }
 
-TArray<UDataLayerInstance*> FDataLayerMode::GetDataLayersFromOperation(const FDragDropOperation& Operation, bool bOnlyFindFirst /*= false*/) const
+FSceneOutlinerDragValidationInfo FDataLayerMode::ValidateDataLayerAssetDrop(const ISceneOutlinerTreeItem& DropTarget, const TArray<const UDataLayerAsset*>& DataLayerAssetsToDrop) const
 {
-	TArray<UDataLayerInstance*> OutDataLayers;
+	check(!DataLayerAssetsToDrop.IsEmpty());
+
+	if (DataLayerEditorSubsystem->HasDeprecatedDataLayers())
+	{
+		return FSceneOutlinerDragValidationInfo(ESceneOutlinerDropCompatibility::Incompatible, FText::Format(LOCTEXT("CantCreateInstanceWorldHasDeprecatedDataLayers", "Cannot create Data Layers Instance from assets since \"{0}\" has deprecated data layers."),
+			FText::FromString(GetOwningWorld()->GetName())));
+	}
+
+	const FDataLayerTreeItem* DataLayerTargetItem = DropTarget.CastTo<FDataLayerTreeItem>();
+	const UDataLayerInstance* DataLayerTarget = DataLayerTargetItem ? DataLayerTargetItem->GetDataLayer() : nullptr;
+	const UDataLayerInstanceWithAsset* DataLayerWithAssetTarget = DataLayerTarget ? CastChecked<UDataLayerInstanceWithAsset>(DataLayerTarget) : nullptr;
+
+	if (!DropTarget.CanInteract() || (DataLayerWithAssetTarget != nullptr && DataLayerWithAssetTarget->IsReadOnly()))
+	{
+		return FSceneOutlinerDragValidationInfo(ESceneOutlinerDropCompatibility::Incompatible, FText::Format(LOCTEXT("DropTargetIsReadOnly", "\"{0}\" is read only."), FText::FromString(DropTarget.GetDisplayString())));
+	}
+
+	TArray<const UDataLayerInstance*> ExistingDataLayerInstance;
+	ExistingDataLayerInstance = DataLayerEditorSubsystem->GetDataLayerInstances(DataLayerAssetsToDrop);
+
+	if (ExistingDataLayerInstance.IsEmpty())
+	{
+		if (DropTarget.CastTo<FDataLayerTreeItem>() || DropTarget.CastTo<FWorldDataLayersTreeItem>())
+		{
+			return FSceneOutlinerDragValidationInfo(ESceneOutlinerDropCompatibility::Compatible, FText::Format(LOCTEXT("CreateAllDataLayerFromAssetDropOnPart", "Create Data Layer Instances under \"{0}\""), FText::FromString(DropTarget.GetDisplayString())));
+		}
+		else
+		{
+			return FSceneOutlinerDragValidationInfo(ESceneOutlinerDropCompatibility::Compatible, LOCTEXT("CreateAllDataLayerFromAssetDrop", "Create Data Layer Instances"));
+		}
+	}
+	else
+	{
+		TArray<FString> ExistingAssestNames;
+		Algo::Transform(ExistingDataLayerInstance, ExistingAssestNames, [](const UDataLayerInstance* DataLayerInstance) { return CastChecked<UDataLayerInstanceWithAsset>(DataLayerInstance)->GetAsset()->GetName(); });
+		
+		FString ExistingAssetNamesString = FString::Join(ExistingAssestNames, TEXT(","));
+
+		return FSceneOutlinerDragValidationInfo(ESceneOutlinerDropCompatibility::Incompatible, FText::Format(LOCTEXT("CantCreateDataLayerInstanceSameAsset", "Cannot create Data Layer Instance(s) there are already Data Layer Instance(s) for \"{0}\"."),
+			FText::FromString(ExistingAssetNamesString)));
+	}
+}
+
+TArray<UDataLayerInstance*> FDataLayerMode::GetDataLayerInstancesFromOperation(const FDragDropOperation& Operation, bool bOnlyFindFirst /*= false*/) const
+{
+	TArray<UDataLayerInstance*> OutDataLayerInstances;
 	
-	auto GetDataLayers = [this, &OutDataLayers](const FDataLayerDragDropOp& DataLayerOp)
+	auto GetDataLayers = [this, &OutDataLayerInstances](const FDataLayerDragDropOp& DataLayerOp)
 	{
 		const TArray<FDataLayerDragDropOp::FDragDropInfo>& DragDropInfos = DataLayerOp.DataLayerDragDropInfos;
 		for (const FDataLayerDragDropOp::FDragDropInfo& DragDropInfo : DragDropInfos)
 		{
 			if (UDataLayerInstance* DataLayerInstance = DataLayerEditorSubsystem->GetDataLayerInstance(DragDropInfo.DataLayerInstanceName))
 			{
-				OutDataLayers.AddUnique(DataLayerInstance);
+				OutDataLayerInstances.AddUnique(DataLayerInstance);
 			}
 		}
 	};
@@ -673,7 +730,26 @@ TArray<UDataLayerInstance*> FDataLayerMode::GetDataLayersFromOperation(const FDr
 		GetDataLayers(DataLayerDragDropOp);
 	}
 
-	return OutDataLayers;
+	return OutDataLayerInstances;
+}
+
+TArray<const UDataLayerAsset*> FDataLayerMode::GetDataLayerAssetsFromOperation(const FDragDropOperation& InDragDrop, bool bOnlyFindFirst /* = false */) const
+{
+	TArray<const UDataLayerAsset*> OutDataLayerAssets;
+
+	TArray<FAssetData> AssetDatas = AssetUtil::ExtractAssetDataFromDrag(InDragDrop.AsShared());
+	for (const FAssetData& AssetData : AssetDatas)
+	{
+		if (AssetData.GetClass()->IsChildOf<UDataLayerAsset>())
+		{
+			if (UObject* Asset = AssetData.GetAsset())
+			{
+				OutDataLayerAssets.Add(StaticCast<UDataLayerAsset*>(Asset));
+			}
+		}
+	}
+
+	return OutDataLayerAssets;
 }
 
 TArray<FDataLayerActorMoveElement> FDataLayerMode::GetDataLayerActorPairsFromOperation(const FDragDropOperation& Operation) const
@@ -764,7 +840,7 @@ TArray<AActor*> FDataLayerMode::GetActorsFromOperation(const FDragDropOperation&
 
 void FDataLayerMode::OnDrop(ISceneOutlinerTreeItem& DropTarget, const FSceneOutlinerDragDropPayload& Payload, const FSceneOutlinerDragValidationInfo& ValidationInfo) const
 {
-	const FDataLayerTreeItem* DataLayerItem = DropTarget.CastTo<FDataLayerTreeItem>();
+	FDataLayerTreeItem* DataLayerItem = DropTarget.CastTo<FDataLayerTreeItem>();
 	UDataLayerInstance* TargetDataLayer = DataLayerItem ? DataLayerItem->GetDataLayer() : nullptr;
 
 	TArray<AActor*> ActorsToAdd = GetActorsFromOperation(Payload.SourceOperation);
@@ -786,8 +862,11 @@ void FDataLayerMode::OnDrop(ISceneOutlinerTreeItem& DropTarget, const FSceneOutl
 			const FScopedTransaction Transaction(LOCTEXT("DataLayerOutlinerAddActorsToDataLayer", "Add Actors to Data Layer"));
 			DataLayerEditorSubsystem->AddActorsToDataLayer(ActorsToAdd, TargetDataLayer);
 		}
+
+		return;
 	}
-	else if (Payload.SourceOperation.IsOfType<FDataLayerActorMoveOp>()) // Moving actor(s) into a Data Layer
+	
+	if (Payload.SourceOperation.IsOfType<FDataLayerActorMoveOp>()) // Moving actor(s) into a Data Layer
 	{
 		TArray<FDataLayerActorMoveElement> ActorsToMove = GetDataLayerActorPairsFromOperation(Payload.SourceOperation);
 		if (!ActorsToMove.IsEmpty() && TargetDataLayer)
@@ -805,13 +884,47 @@ void FDataLayerMode::OnDrop(ISceneOutlinerTreeItem& DropTarget, const FSceneOutl
 
 		return;
 	}
-	else // Moving a data layer(s)
+
+	TArray<const UDataLayerAsset*> DataLayerAssetToDrop = GetDataLayerAssetsFromOperation(Payload.SourceOperation);
+	if (!DataLayerAssetToDrop.IsEmpty()) // Dropping Data Layer Assets 
+	{
+		OnDataLayerAssetDropped(DataLayerAssetToDrop, DropTarget);
+		return;
+	}
+	
+	// Moving a data layer(s)
 	{
-		TArray<UDataLayerInstance*> DataLayerInstances = GetDataLayersFromOperation(Payload.SourceOperation);
+		TArray<UDataLayerInstance*> DataLayerInstances = GetDataLayerInstancesFromOperation(Payload.SourceOperation);
 		SetParentDataLayer(DataLayerInstances, TargetDataLayer);
 	}
 }
 
+void FDataLayerMode::OnDataLayerAssetDropped(const TArray<const UDataLayerAsset*>& DroppedDataLayerAsset, ISceneOutlinerTreeItem& DropTarget) const
+{
+	const FScopedTransaction Transaction(LOCTEXT("DataLayerOutlinerCreateDataLayerInstanceFromAssetDrop", "Create Data Layer Instance from Data Layer Asset Drop"));
+
+	for (const UDataLayerAsset* DataLayerAsset : DroppedDataLayerAsset)
+	{
+		FDataLayerCreationParameters Params;
+		Params.DataLayerAsset = const_cast<UDataLayerAsset*>(DataLayerAsset);
+
+		if (FWorldDataLayersTreeItem* WorldDataLayerTreeItem = DropTarget.CastTo<FWorldDataLayersTreeItem>())
+		{
+			Params.WorldDataLayers = WorldDataLayerTreeItem->GetWorldDataLayers();
+		}
+
+		UDataLayerInstance* DataLayerInstance = DataLayerEditorSubsystem->CreateDataLayerInstance(Params);
+
+		if (FDataLayerTreeItem* DataLayerItem = DropTarget.CastTo<FDataLayerTreeItem>())
+		{
+			if (UDataLayerInstance* ParentDataLayer = DataLayerItem->GetDataLayer())
+			{
+				DataLayerEditorSubsystem->SetParentDataLayer(DataLayerInstance, ParentDataLayer);
+			}
+		}
+	}
+}
+
 FReply FDataLayerMode::OnDragOverItem(const FDragDropEvent& Event, const ISceneOutlinerTreeItem& Item) const
 {
 	TSharedPtr<FDragDropOperation> DragOperation = Event.GetOperation();
diff --git a/Engine/Source/Editor/DataLayerEditor/Private/DataLayer/DataLayerMode.h b/Engine/Source/Editor/DataLayerEditor/Private/DataLayer/DataLayerMode.h
--- a/Engine/Source/Editor/DataLayerEditor/Private/DataLayer/DataLayerMode.h
+++ b/Engine/Source/Editor/DataLayerEditor/Private/DataLayer/DataLayerMode.h
@@ -137,7 +137,8 @@ private:
 	bool IsWorldChecked(TWeakObjectPtr<UWorld> World) const;
 	TArray<FDataLayerActorMoveElement> GetDataLayerActorPairsFromOperation(const FDragDropOperation& Operation) const;
 	TArray<AActor*> GetActorsFromOperation(const FDragDropOperation& Operation, bool bOnlyFindFirst = false) const;
-	TArray<UDataLayerInstance*> GetDataLayersFromOperation(const FDragDropOperation& Operation, bool bOnlyFindFirst = false) const;
+	TArray<UDataLayerInstance*> GetDataLayerInstancesFromOperation(const FDragDropOperation& Operation, bool bOnlyFindFirst = false) const;
+	TArray<const UDataLayerAsset*> GetDataLayerAssetsFromOperation(const FDragDropOperation& InDragDrop, bool bOnlyFindFirst = false) const;
 	TArray<UDataLayerInstance*> GetSelectedDataLayers(SSceneOutliner* InSceneOutliner) const;
 	TArray<AWorldDataLayers*> GetSelectedWorldDataLayers(SSceneOutliner* InSceneOutliner) const;
 	void SetParentDataLayer(const TArray<UDataLayerInstance*> DataLayers, UDataLayerInstance* ParentDataLayer) const;
@@ -148,7 +149,9 @@ private:
 	void RefreshSelection();
 	UWorld* GetOwningWorld() const;
 	AWorldDataLayers* GetOwningWorldAWorldDataLayers() const;
-	FSceneOutlinerDragValidationInfo ValidateDrop(const ISceneOutlinerTreeItem& DropTarget, bool bMoveOperation = false) const;
+	FSceneOutlinerDragValidationInfo ValidateActorDrop(const ISceneOutlinerTreeItem& DropTarget, bool bMoveOperation = false) const;
+	FSceneOutlinerDragValidationInfo ValidateDataLayerAssetDrop(const ISceneOutlinerTreeItem& DropTarget, const TArray<const UDataLayerAsset*>& DataLayerAssetsToDrop) const;
+	void OnDataLayerAssetDropped(const TArray<const UDataLayerAsset*>& DroppedDataLayerAsset, ISceneOutlinerTreeItem& DropTarget) const;
 
 	/** Filter factories */
 	static TSharedRef<FSceneOutlinerFilter> CreateShowOnlySelectedActorsFilter();
diff --git a/Engine/Source/Editor/DataLayerEditor/Public/DataLayer/DataLayerEditorSubsystem.h b/Engine/Source/Editor/DataLayerEditor/Public/DataLayer/DataLayerEditorSubsystem.h
--- a/Engine/Source/Editor/DataLayerEditor/Public/DataLayer/DataLayerEditorSubsystem.h
+++ b/Engine/Source/Editor/DataLayerEditor/Public/DataLayer/DataLayerEditorSubsystem.h
@@ -113,8 +113,10 @@ public:
 	void OnWorldDataLayersPostRegister(AWorldDataLayers* WorldDataLayers);
 	void OnWorldDataLayersPreUnregister(AWorldDataLayers* WorldDataLayers);
 
+	TArray<const UDataLayerInstance*> GetDataLayerInstances(const TArray<const UDataLayerAsset*> DataLayerAssets) const;
+
 	template<class DataLayerType, typename ...Args>
-	UDataLayerInstance* CreateDataLayerInstance(AWorldDataLayers* WorldDataLayers, Args&&... InArgs);
+	DataLayerType* CreateDataLayerInstance(AWorldDataLayers* WorldDataLayers, Args&&... InArgs);
 
 	/* Broadcasts whenever one or more DataLayers are modified
 	 *
@@ -526,14 +528,23 @@ public:
 	void MakeAllDataLayersVisible();
 
 	/**
-	 * Gets the UDataLayerInstance Object from associate with the DataLayerAsset
+	 * Gets the UDataLayerInstance associated to the DataLayerAsset
 	 *
 	 * @param	DataLayerAsset	The DataLayerAsset associated to the UDataLayerInstance
-	 * @return					The UDataLayerInstance Object of the provided DataLayerAsset
+	 * @return					The UDataLayerInstance of the provided DataLayerAsset
 	 */
 	UFUNCTION(BlueprintCallable, Category = DataLayers)
 	UDataLayerInstance* GetDataLayerInstance(const UDataLayerAsset* DataLayerAsset) const;
 
+	/**
+	 * Gets the UDataLayerInstances associated to the each DataLayerAssets
+	 *
+	 * @param	DataLayerAssets	The array of DataLayerAssets associated to UDataLayerInstances
+	 * @return					The array of UDataLayerInstances corresponding to a DataLayerAsset in the DataLayerAssets array
+	 */
+	UFUNCTION(BlueprintCallable, Category = DataLayers)
+	TArray<UDataLayerInstance*> GetDataLayerInstances(const TArray<UDataLayerAsset*> DataLayerAssets) const;
+
 	/**
 	 * Gets the UDataLayerInstance Object of the DataLayer name
 	 *
@@ -745,7 +756,7 @@ private:
 };
 
 template<class DataLayerType, typename ...Args>
-UDataLayerInstance* UDataLayerEditorSubsystem::CreateDataLayerInstance(AWorldDataLayers* WorldDataLayers, Args&&... InArgs)
+DataLayerType* UDataLayerEditorSubsystem::CreateDataLayerInstance(AWorldDataLayers* WorldDataLayers, Args&&... InArgs)
 {
 	UDataLayerInstance* NewDataLayer = nullptr;
 
@@ -759,5 +770,5 @@ UDataLayerInstance* UDataLayerEditorSubsystem::CreateDataLayerInstance(AWorldDat
 		BroadcastDataLayerChanged(EDataLayerAction::Add, NewDataLayer, NAME_None);
 	}
 
-	return NewDataLayer;
+	return CastChecked<DataLayerType>(NewDataLayer);
 }
\ No newline at end of file
diff --git a/Engine/Source/Editor/UnrealEd/Private/AssetSelection.cpp b/Engine/Source/Editor/UnrealEd/Private/AssetSelection.cpp
--- a/Engine/Source/Editor/UnrealEd/Private/AssetSelection.cpp
+++ b/Engine/Source/Editor/UnrealEd/Private/AssetSelection.cpp
@@ -678,7 +678,7 @@ namespace AssetUtil
 		return ExtractAssetDataFromDrag(DragDropEvent.GetOperation());
 	}
 
-	TArray<FAssetData> ExtractAssetDataFromDrag(const TSharedPtr<FDragDropOperation>& Operation)
+	TArray<FAssetData> ExtractAssetDataFromDrag(const TSharedPtr<const FDragDropOperation>& Operation)
 	{
 		TArray<FAssetData> DroppedAssetData;
 
@@ -689,8 +689,8 @@ namespace AssetUtil
 
 		if (Operation->IsOfType<FExternalDragOperation>())
 		{
-			TSharedPtr<FExternalDragOperation> DragDropOp = StaticCastSharedPtr<FExternalDragOperation>(Operation);
-			if ( DragDropOp->HasText() )
+			TSharedPtr<const FExternalDragOperation> DragDropOp = StaticCastSharedPtr<const FExternalDragOperation>(Operation);
+			if (DragDropOp->HasText())
 			{
 				TArray<FString> DroppedAssetStrings;
 				const TCHAR AssetDelimiter[] = { AssetMarshalDefs::AssetDelimiter, TEXT('\0') };
@@ -714,13 +714,13 @@ namespace AssetUtil
 		}
 		else if (Operation->IsOfType<FCollectionDragDropOp>())
 		{
-			TSharedPtr<FCollectionDragDropOp> DragDropOp = StaticCastSharedPtr<FCollectionDragDropOp>( Operation );
-			DroppedAssetData.Append( DragDropOp->GetAssets() );
+			TSharedPtr<const FCollectionDragDropOp> DragDropOp = StaticCastSharedPtr<const FCollectionDragDropOp>(Operation);
+			DroppedAssetData.Append(DragDropOp->GetAssets());
 		}
 		else if (Operation->IsOfType<FAssetDragDropOp>())
 		{
-			TSharedPtr<FAssetDragDropOp> DragDropOp = StaticCastSharedPtr<FAssetDragDropOp>( Operation );
-			DroppedAssetData.Append( DragDropOp->GetAssets() );
+			TSharedPtr<const FAssetDragDropOp> DragDropOp = StaticCastSharedPtr<const FAssetDragDropOp>(Operation);
+			DroppedAssetData.Append(DragDropOp->GetAssets());
 		}
 
 		return DroppedAssetData;
diff --git a/Engine/Source/Editor/UnrealEd/Public/AssetSelection.h b/Engine/Source/Editor/UnrealEd/Public/AssetSelection.h
--- a/Engine/Source/Editor/UnrealEd/Public/AssetSelection.h
+++ b/Engine/Source/Editor/UnrealEd/Public/AssetSelection.h
@@ -26,8 +26,7 @@ namespace AssetUtil
 	/** Extract the names of the assets which are being dragged*/
 	UNREALED_API TArray<FAssetData>	ExtractAssetDataFromDrag( const FDragDropEvent& DragDropEvent );
 
-	/** Extract the names of the assets which are being dragged*/
-	UNREALED_API TArray<FAssetData> ExtractAssetDataFromDrag( const TSharedPtr<FDragDropOperation>& Operation );
+	UNREALED_API TArray<FAssetData> ExtractAssetDataFromDrag(const TSharedPtr<const FDragDropOperation>& Operation);
 
 	/**
 	 *  Given an array of asset names, loads the assets into an array of objects
diff --git a/Engine/Source/Runtime/Engine/Private/WorldPartition/DataLayer/WorldDataLayers.cpp b/Engine/Source/Runtime/Engine/Private/WorldPartition/DataLayer/WorldDataLayers.cpp
--- a/Engine/Source/Runtime/Engine/Private/WorldPartition/DataLayer/WorldDataLayers.cpp
+++ b/Engine/Source/Runtime/Engine/Private/WorldPartition/DataLayer/WorldDataLayers.cpp
@@ -464,7 +464,7 @@ AWorldDataLayers* AWorldDataLayers::Create(const FActorSpawnParameters& SpawnPar
 	return WorldDataLayers;
 }
 
-TArray<FName> AWorldDataLayers::GetDataLayerInstanceNames(const TArray<TObjectPtr<const UDataLayerAsset>>& InDataLayersAssets) const
+TArray<FName> AWorldDataLayers::GetDataLayerInstanceNames(const TArray<const UDataLayerAsset*>& InDataLayersAssets) const
 {
 	TArray<FName> OutDataLayerNames;
 	OutDataLayerNames.Reserve(InDataLayersAssets.Num());
@@ -477,14 +477,14 @@ TArray<FName> AWorldDataLayers::GetDataLayerInstanceNames(const TArray<TObjectPt
 	return OutDataLayerNames;
 }
 
-TArray<const UDataLayerInstance*> AWorldDataLayers::GetDataLayerInstances(const TArray<TObjectPtr<const UDataLayerAsset>>& InDataLayersAssets) const
+TArray<const UDataLayerInstance*> AWorldDataLayers::GetDataLayerInstances(const TArray<const UDataLayerAsset*>& InDataLayersAssets) const
 {
 	TArray<const UDataLayerInstance*> OutDataLayers;
 	OutDataLayers.Reserve(InDataLayersAssets.Num());
 
-	for (const TObjectPtr<const UDataLayerAsset>& DataLayerAsset : InDataLayersAssets)
+	for (const UDataLayerAsset* DataLayerAsset : InDataLayersAssets)
 	{
-		if (const UDataLayerInstance* DataLayerObject = GetDataLayerInstance(DataLayerAsset.Get()))
+		if (const UDataLayerInstance* DataLayerObject = GetDataLayerInstance(DataLayerAsset))
 		{
 			OutDataLayers.AddUnique(DataLayerObject);
 		}
diff --git a/Engine/Source/Runtime/Engine/Public/WorldPartition/DataLayer/WorldDataLayers.h b/Engine/Source/Runtime/Engine/Public/WorldPartition/DataLayer/WorldDataLayers.h
--- a/Engine/Source/Runtime/Engine/Public/WorldPartition/DataLayer/WorldDataLayers.h
+++ b/Engine/Source/Runtime/Engine/Public/WorldPartition/DataLayer/WorldDataLayers.h
@@ -79,8 +79,9 @@ public:
 	TArray<UDataLayerInstance*> GetActorEditorContextDataLayers() const;
 
 	//~ Begin Helper Functions
-	TArray<const UDataLayerInstance*> GetDataLayerInstances(const TArray<TObjectPtr<const UDataLayerAsset>>& InDataLayersAssets) const;
-	TArray<FName> GetDataLayerInstanceNames(const TArray<TObjectPtr<const UDataLayerAsset>>& InDataLayersAssets) const;
+	TArray<const UDataLayerInstance*> GetDataLayerInstances(const TArray<UDataLayerAsset*>& InDataLayersAssets) const;
+	TArray<const UDataLayerInstance*> GetDataLayerInstances(const TArray<const UDataLayerAsset*>& InDataLayersAssets) const;
+	TArray<FName> GetDataLayerInstanceNames(const TArray<const UDataLayerAsset*>& InDataLayersAssets) const;
 	//~ End Helper Functions
 
 	// Allows overriding of DataLayers with PlayFromHere
