Fixed check when using "Make Current Data Layer" after deleting a Data Layer in the outliner.

diff --git a/Engine/Source/Editor/DataLayerEditor/Private/DataLayer/DataLayerMode.cpp b/Engine/Source/Editor/DataLayerEditor/Private/DataLayer/DataLayerMode.cpp
--- a/Engine/Source/Editor/DataLayerEditor/Private/DataLayer/DataLayerMode.cpp
+++ b/Engine/Source/Editor/DataLayerEditor/Private/DataLayer/DataLayerMode.cpp
@@ -473,21 +473,26 @@ void FDataLayerMode::DeleteItems(const TArray<TWeakPtr<ISceneOutlinerTreeItem>>&
 	}
 	else if (!DataLayersToDelete.IsEmpty())
 	{
-		int32 PrevDeleteCount = SelectedDataLayersSet.Num();
-		for (UDataLayerInstance* DataLayerToDelete : DataLayersToDelete)
-		{
-			SelectedDataLayersSet.Remove(DataLayerToDelete);
-		}
-		
-		{
-			const FScopedTransaction Transaction(LOCTEXT("DeleteDataLayers", "Delete Data Layers"));
-			DataLayerEditorSubsystem->DeleteDataLayers(DataLayersToDelete);
-		}
+		DeleteDataLayers(DataLayersToDelete);
+	}
+}
 
-		if ((SelectedDataLayersSet.Num() != PrevDeleteCount) && DataLayerBrowser)
-		{
-			DataLayerBrowser->OnSelectionChanged(SelectedDataLayersSet);
-		}
+void FDataLayerMode::DeleteDataLayers(const TArray<UDataLayerInstance*>& InDataLayersToDelete)
+{
+	int32 PrevDeleteCount = SelectedDataLayersSet.Num();
+	for (UDataLayerInstance* DataLayerToDelete : InDataLayersToDelete)
+	{
+		SelectedDataLayersSet.Remove(DataLayerToDelete);
+	}
+
+	{
+		const FScopedTransaction Transaction(LOCTEXT("DeleteSelectedDataLayers", "Delete Selected Data Layers"));
+		DataLayerEditorSubsystem->DeleteDataLayers(InDataLayersToDelete);
+	}
+
+	if ((SelectedDataLayersSet.Num() != PrevDeleteCount) && DataLayerBrowser)
+	{
+		DataLayerBrowser->OnSelectionChanged(SelectedDataLayersSet);
 	}
 }
 
@@ -1359,11 +1364,10 @@ void FDataLayerMode::RegisterContextMenu()
 
 				Section.AddMenuEntry("DeleteSelectedDataLayers", LOCTEXT("DeleteSelectedDataLayers", "Delete Selected Data Layers"), FText(), FSlateIcon(),
 					FUIAction(
-						FExecuteAction::CreateLambda([SelectedDataLayers]() {
+						FExecuteAction::CreateLambda([Mode, SelectedDataLayers]() {
 							check(!SelectedDataLayers.IsEmpty());
 							{
-								const FScopedTransaction Transaction(LOCTEXT("DeleteSelectedDataLayers", "Delete Selected Data Layers"));
-								UDataLayerEditorSubsystem::Get()->DeleteDataLayers(SelectedDataLayers);
+								Mode->DeleteDataLayers(SelectedDataLayers);
 							}}),
 						FCanExecuteAction::CreateLambda([SelectedDataLayers,bSelectedDataLayersContainsLocked] { return !SelectedDataLayers.IsEmpty() && !bSelectedDataLayersContainsLocked; })
 					));
diff --git a/Engine/Source/Editor/DataLayerEditor/Private/DataLayer/DataLayerMode.h b/Engine/Source/Editor/DataLayerEditor/Private/DataLayer/DataLayerMode.h
--- a/Engine/Source/Editor/DataLayerEditor/Private/DataLayer/DataLayerMode.h
+++ b/Engine/Source/Editor/DataLayerEditor/Private/DataLayer/DataLayerMode.h
@@ -146,6 +146,7 @@ private:
 	static void CreateDataLayerPicker(UToolMenu* InMenu, FOnDataLayerPicked OnDataLayerPicked, bool bInShowRoot = false);
 	bool ShouldExpandDataLayer(const UDataLayerInstance* DataLayer) const;
 	bool ContainsSelectedChildDataLayer(const UDataLayerInstance* DataLayer) const;
+	void DeleteDataLayers(const TArray<UDataLayerInstance*>& InDataLayersToDelete);
 	void RefreshSelection();
 	UWorld* GetOwningWorld() const;
 	AWorldDataLayers* GetOwningWorldAWorldDataLayers() const;
