fix file indentation

#rnx

diff --git a/Engine/Source/Editor/DataLayerEditor/Private/DataLayer/DataLayerEditorSubsystem.cpp b/Engine/Source/Editor/DataLayerEditor/Private/DataLayer/DataLayerEditorSubsystem.cpp
--- a/Engine/Source/Editor/DataLayerEditor/Private/DataLayer/DataLayerEditorSubsystem.cpp
+++ b/Engine/Source/Editor/DataLayerEditor/Private/DataLayer/DataLayerEditorSubsystem.cpp
@@ -1146,9 +1146,9 @@ bool UDataLayerEditorSubsystem::HasDeprecatedDataLayers() const
 UDataLayerInstance* UDataLayerEditorSubsystem::GetDataLayerInstance(const FName& DataLayerInstanceName) const
 {
 	if (UDataLayerSubsystem* DataLayerSubsystem = UWorld::GetSubsystem<UDataLayerSubsystem>(GetWorld()))
-{
+	{
 		return DataLayerSubsystem->GetDataLayerInstance(DataLayerInstanceName);
-}
+	}
 	return nullptr;
 }
 
@@ -1277,20 +1277,20 @@ void UDataLayerEditorSubsystem::OnSelectionChanged()
 const TSet<TWeakObjectPtr<const UDataLayerInstance>>& UDataLayerEditorSubsystem::GetSelectedDataLayersFromEditorSelection() const
 {
 	if (bRebuildSelectedDataLayersFromEditorSelection)
-{
+	{
 		bRebuildSelectedDataLayersFromEditorSelection = false;
 
-	SelectedDataLayersFromEditorSelection.Reset();
-	TArray<AActor*> Actors;
-	GEditor->GetSelectedActors()->GetSelectedObjects<AActor>(Actors);
-	for (const AActor* Actor : Actors)
-	{
-		for (const UDataLayerInstance* DataLayerInstance : Actor->GetDataLayerInstances())
+		SelectedDataLayersFromEditorSelection.Reset();
+		TArray<AActor*> Actors;
+		GEditor->GetSelectedActors()->GetSelectedObjects<AActor>(Actors);
+		for (const AActor* Actor : Actors)
 		{
-			SelectedDataLayersFromEditorSelection.Add(DataLayerInstance);
+			for (const UDataLayerInstance* DataLayerInstance : Actor->GetDataLayerInstances())
+			{
+				SelectedDataLayersFromEditorSelection.Add(DataLayerInstance);
+			}
 		}
 	}
-}
 	return SelectedDataLayersFromEditorSelection;
 }
 
