Data Layer outliner double click on Data Layers now adds or removes the Data Layer to/from the editing context (same behavior as Content Bundle outliner).

diff --git a/Engine/Source/Editor/DataLayerEditor/Private/DataLayer/DataLayerMode.cpp b/Engine/Source/Editor/DataLayerEditor/Private/DataLayer/DataLayerMode.cpp
--- a/Engine/Source/Editor/DataLayerEditor/Private/DataLayer/DataLayerMode.cpp
+++ b/Engine/Source/Editor/DataLayerEditor/Private/DataLayer/DataLayerMode.cpp
@@ -412,13 +412,23 @@ void FDataLayerMode::OnItemPassesFilters(const ISceneOutlinerTreeItem& Item)
 
 void FDataLayerMode::OnItemDoubleClick(FSceneOutlinerTreeItemPtr Item)
 {
-	if (const FDataLayerTreeItem* DataLayerItem = Item->CastTo<FDataLayerTreeItem>())
+	if (FDataLayerTreeItem* DataLayerItem = Item->CastTo<FDataLayerTreeItem>())
 	{
 		if (UDataLayerInstance* DataLayerInstance = DataLayerItem->GetDataLayer())
 		{
-			const FScopedTransaction Transaction(LOCTEXT("SelectActorsInDataLayer", "Select Actors in Data Layer"));
-			GEditor->SelectNone(/*bNoteSelectionChange*/false, true);
-			DataLayerEditorSubsystem->SelectActorsInDataLayer(DataLayerInstance, /*bSelect*/true, /*bNotify*/true, /*bSelectEvenIfHidden*/true);
+			if (!DataLayerInstance->IsLocked())
+			{
+				if (!DataLayerInstance->IsInActorEditorContext())
+				{
+					const FScopedTransaction Transaction(LOCTEXT("MakeCurrentDataLayers", "Make Current Data Layer(s)"));
+					UDataLayerEditorSubsystem::Get()->AddToActorEditorContext(DataLayerInstance);
+				}
+				else
+				{
+					const FScopedTransaction Transaction(LOCTEXT("RemoveCurrentDataLayers", "Remove Current Data Layer(s)"));
+					UDataLayerEditorSubsystem::Get()->RemoveFromActorEditorContext(DataLayerInstance);
+				}
+			}
 		}
 	}
 	else if (FDataLayerActorTreeItem* DataLayerActorItem = Item->CastTo<FDataLayerActorTreeItem>())
