Data Layer outliner QoL : Removed Data Layer AccentBlue color for Data Layers used by actor selection (kept bold font) so that current editing Data Layers always keep their green color.

diff --git a/Engine/Source/Editor/DataLayerEditor/Private/DataLayer/SDataLayerTreeLabel.cpp b/Engine/Source/Editor/DataLayerEditor/Private/DataLayer/SDataLayerTreeLabel.cpp
--- a/Engine/Source/Editor/DataLayerEditor/Private/DataLayer/SDataLayerTreeLabel.cpp
+++ b/Engine/Source/Editor/DataLayerEditor/Private/DataLayer/SDataLayerTreeLabel.cpp
@@ -288,11 +288,7 @@ FSlateColor SDataLayerTreeLabel::GetForegroundColor() const
 	{
 		return FLinearColor(0.2f, 0.2f, 0.25f);
 	}
-	if (ShouldBeHighlighted())
-	{
-		return FAppStyle::Get().GetSlateColor("Colors.AccentBlue");
-	}
-	else if (IsInActorEditorContext())
+	if (IsInActorEditorContext())
 	{
 		return FAppStyle::Get().GetSlateColor("Colors.AccentGreen");
 	}
