Added "Browse to asset" on Data Layer items in Data Layer Outliner.

diff --git a/Engine/Source/Editor/DataLayerEditor/Private/DataLayer/DataLayerMode.cpp b/Engine/Source/Editor/DataLayerEditor/Private/DataLayer/DataLayerMode.cpp
--- a/Engine/Source/Editor/DataLayerEditor/Private/DataLayer/DataLayerMode.cpp
+++ b/Engine/Source/Editor/DataLayerEditor/Private/DataLayer/DataLayerMode.cpp
@@ -37,6 +37,7 @@
 #include "Framework/Application/SlateApplication.h"
 #include "Framework/Commands/UIAction.h"
 #include "Framework/Commands/UICommandInfo.h"
+#include "Framework/Commands/UICommandList.h"
 #include "Framework/MultiBox/MultiBoxBuilder.h"
 #include "GameFramework/Actor.h"
 #include "GenericPlatform/ICursor.h"
@@ -67,6 +68,7 @@
 #include "Templates/TypeHash.h"
 #include "Templates/UnrealTemplate.h"
 #include "Textures/SlateIcon.h"
+#include "Toolkits/GlobalEditorCommonCommands.h"
 #include "ToolMenu.h"
 #include "ToolMenuContext.h"
 #include "ToolMenuDelegates.h"
@@ -115,6 +117,13 @@ FDataLayerMode::FDataLayerMode(const FDataLayerModeParams& Params)
 	, SpecifiedWorldToDisplay(Params.SpecifiedWorldToDisplay)
 	, FilteredDataLayerCount(0)
 {
+
+	Commands = MakeShareable(new FUICommandList());
+	Commands->MapAction(FGlobalEditorCommonCommands::Get().FindInContentBrowser, FUIAction(
+		FExecuteAction::CreateRaw(this, &FDataLayerMode::FindInContentBrowser),
+		FCanExecuteAction::CreateRaw(this, &FDataLayerMode::CanFindInContentBrowser)
+	));
+
 	USelection::SelectionChangedEvent.AddRaw(this, &FDataLayerMode::OnLevelSelectionChanged);
 	USelection::SelectObjectEvent.AddRaw(this, &FDataLayerMode::OnLevelSelectionChanged);
 
@@ -539,6 +548,12 @@ FReply FDataLayerMode::OnKeyDown(const FKeyEvent& InKeyEvent)
 		return FReply::Handled();
 
 	}
+
+	if (Commands->ProcessCommandBindings(InKeyEvent))
+	{
+		return FReply::Handled();
+	}
+
 	return FReply::Unhandled();
 }
 
@@ -1170,6 +1185,43 @@ AWorldDataLayers* FDataLayerMode::GetOwningWorldAWorldDataLayers() const
 	return OwningWorld ? OwningWorld->GetWorldDataLayers() : nullptr;
 }
 
+void FDataLayerMode::FindInContentBrowser()
+{
+	if (SceneOutliner)
+	{
+		TArray<UObject*> Objects;
+		for (TWeakObjectPtr<const UDataLayerInstance>& DataLayerInstance : SelectedDataLayersSet)
+		{
+			if (const UDataLayerInstanceWithAsset* DataLayerInstanceWithAsset = Cast<UDataLayerInstanceWithAsset>(DataLayerInstance.Get()))
+			{
+				if (const UDataLayerAsset* Asset = DataLayerInstanceWithAsset->GetAsset())
+				{
+					Objects.Add(const_cast<UDataLayerAsset*>(Asset));
+				}
+			}
+		}
+		if (!Objects.IsEmpty())
+		{
+			GEditor->SyncBrowserToObjects(Objects);
+		}
+	}
+}
+
+bool FDataLayerMode::CanFindInContentBrowser() const
+{
+	for (const TWeakObjectPtr<const UDataLayerInstance>& DataLayerInstance : SelectedDataLayersSet)
+	{
+		if (const UDataLayerInstanceWithAsset* DataLayerInstanceWithAsset = Cast<UDataLayerInstanceWithAsset>(DataLayerInstance.Get()))
+		{
+			if (const UDataLayerAsset* Asset = DataLayerInstanceWithAsset->GetAsset())
+			{
+				return true;
+			}
+		}
+	}
+	return false;
+}
+
 void FDataLayerMode::RegisterContextMenu()
 {
 	UToolMenus* ToolMenus = UToolMenus::Get();
@@ -1178,7 +1230,7 @@ void FDataLayerMode::RegisterContextMenu()
 	{
 		UToolMenu* Menu = ToolMenus->RegisterMenu(DefaultContextBaseMenuName);
 
-		Menu->AddDynamicSection("DataLayerDynamicSection", FNewToolMenuDelegate::CreateLambda([](UToolMenu* InMenu)
+		Menu->AddDynamicSection("DataLayerDynamicSection", FNewToolMenuDelegate::CreateLambda([this](UToolMenu* InMenu)
 		{
 			USceneOutlinerMenuContext* Context = InMenu->FindContext<USceneOutlinerMenuContext>();
 			if (!Context || !Context->SceneOutliner.IsValid())
@@ -1529,6 +1581,12 @@ void FDataLayerMode::RegisterContextMenu()
 						FCanExecuteAction::CreateLambda([AllDataLayers] { return !AllDataLayers.IsEmpty(); })
 					));
 			}
+
+			if (!UDataLayerEditorSubsystem::Get()->HasDeprecatedDataLayers())
+			{
+				FToolMenuSection& Section = InMenu->AddSection("AssetOptionsSection", LOCTEXT("AssetOptionsText", "Asset Options"));
+				Section.AddMenuEntryWithCommandList(FGlobalEditorCommonCommands::Get().FindInContentBrowser, Commands);
+			}
 		}));
 	}
 
diff --git a/Engine/Source/Editor/DataLayerEditor/Private/DataLayer/DataLayerMode.h b/Engine/Source/Editor/DataLayerEditor/Private/DataLayer/DataLayerMode.h
--- a/Engine/Source/Editor/DataLayerEditor/Private/DataLayer/DataLayerMode.h
+++ b/Engine/Source/Editor/DataLayerEditor/Private/DataLayer/DataLayerMode.h
@@ -31,6 +31,7 @@ class FDragDropEvent;
 class FDragDropOperation;
 class FMenuBuilder;
 class FSceneOutlinerFilter;
+class FUICommandList;
 class ISceneOutlinerHierarchy;
 class SDataLayerBrowser;
 class SWidget;
@@ -131,7 +132,7 @@ protected:
 
 private:
 	/* Private Helpers */
-	static void RegisterContextMenu();
+	void RegisterContextMenu();
 	void ChooseRepresentingWorld();
 	void OnSelectWorld(TWeakObjectPtr<UWorld> World);
 	bool IsWorldChecked(TWeakObjectPtr<UWorld> World) const;
@@ -162,6 +163,12 @@ private:
 	static TSharedRef<FSceneOutlinerFilter> CreateHideUnloadedActorsFilter();
 	static TSharedRef<FSceneOutlinerFilter> CreateHideLevelInstancesFilter();
 
+private:
+
+	/** Delegate to handle "Find in Content Browser" context menu option */
+	void FindInContentBrowser();
+	/** Delegate to handle enabling the "Find in Content Browser" context menu option */
+	bool CanFindInContentBrowser() const;
 	/** The world which we are currently representing */
 	TWeakObjectPtr<UWorld> RepresentingWorld;
 	/** The world which the user manually selected */
@@ -178,6 +185,8 @@ private:
 	TSet<TWeakObjectPtr<UDataLayerInstance>> ApplicableDataLayers;
 	/** The path at which the "Pick A Data Layer Asset" will be opened*/
 	mutable FString PickDataLayerDialogPath;
+	/** Command list */
+	TSharedPtr<FUICommandList> Commands;
 
 	TSet<TWeakObjectPtr<const UDataLayerInstance>> SelectedDataLayersSet;
 	typedef TPair<TWeakObjectPtr<const UDataLayerInstance>, TWeakObjectPtr<const AActor>> FSelectedDataLayerActor;
