Removed "Rename Data Layer" menu entry in Data Layer Outliner for maps using Data Layer Assets.

diff --git a/Engine/Source/Editor/DataLayerEditor/Private/DataLayer/DataLayerMode.cpp b/Engine/Source/Editor/DataLayerEditor/Private/DataLayer/DataLayerMode.cpp
--- a/Engine/Source/Editor/DataLayerEditor/Private/DataLayer/DataLayerMode.cpp
+++ b/Engine/Source/Editor/DataLayerEditor/Private/DataLayer/DataLayerMode.cpp
@@ -1434,20 +1434,23 @@ void FDataLayerMode::RegisterContextMenu()
 						FCanExecuteAction::CreateLambda([SelectedDataLayers,bSelectedDataLayersContainsLocked] { return !SelectedDataLayers.IsEmpty() && !bSelectedDataLayersContainsLocked; })
 					));
 
-				Section.AddMenuEntry("RenameSelectedDataLayer", LOCTEXT("RenameSelectedDataLayer", "Rename Selected Data Layer"), FText(), FSlateIcon(),
-					FUIAction(
-						FExecuteAction::CreateLambda([Mode,SceneOutliner,SelectedDataLayers]() {
-							if (SelectedDataLayers.Num() == 1)
-							{
-								FSceneOutlinerTreeItemPtr ItemToRename = SceneOutliner->GetTreeItem(SelectedDataLayers[0]);
-								if (ItemToRename.IsValid() && Mode->CanRenameItem(*ItemToRename) && ItemToRename->CanInteract())
+				if (UDataLayerEditorSubsystem::Get()->HasDeprecatedDataLayers())
+				{
+					Section.AddMenuEntry("RenameSelectedDataLayer", LOCTEXT("RenameSelectedDataLayer", "Rename Selected Data Layer"), FText(), FSlateIcon(),
+						FUIAction(
+							FExecuteAction::CreateLambda([Mode, SceneOutliner, SelectedDataLayers]() {
+								if (SelectedDataLayers.Num() == 1)
 								{
-									SceneOutliner->SetPendingRenameItem(ItemToRename);
-									SceneOutliner->ScrollItemIntoView(ItemToRename);
-								}
-							}}),
-						FCanExecuteAction::CreateLambda([SelectedDataLayers] { return (SelectedDataLayers.Num() == 1) && !SelectedDataLayers[0]->IsLocked(); })
-					));
+									FSceneOutlinerTreeItemPtr ItemToRename = SceneOutliner->GetTreeItem(SelectedDataLayers[0]);
+									if (ItemToRename.IsValid() && Mode->CanRenameItem(*ItemToRename) && ItemToRename->CanInteract())
+									{
+										SceneOutliner->SetPendingRenameItem(ItemToRename);
+										SceneOutliner->ScrollItemIntoView(ItemToRename);
+									}
+								}}),
+							FCanExecuteAction::CreateLambda([SelectedDataLayers] { return (SelectedDataLayers.Num() == 1) && !SelectedDataLayers[0]->IsLocked(); })
+						));
+				}
 
 				Section.AddSeparator("SectionsSeparator");
 			}
