Missing DATALAYEREDITOR_API on FDataLayerCreationParameters.

diff --git a/Engine/Source/Editor/DataLayerEditor/Public/DataLayer/DataLayerEditorSubsystem.h b/Engine/Source/Editor/DataLayerEditor/Public/DataLayer/DataLayerEditorSubsystem.h
--- a/Engine/Source/Editor/DataLayerEditor/Public/DataLayer/DataLayerEditorSubsystem.h
+++ b/Engine/Source/Editor/DataLayerEditor/Public/DataLayer/DataLayerEditorSubsystem.h
@@ -45,7 +45,7 @@ struct FFrame;
 template<typename TItemType> class IFilter;
 
 USTRUCT(BlueprintType)
-struct FDataLayerCreationParameters
+struct DATALAYEREDITOR_API FDataLayerCreationParameters
 {
 	GENERATED_USTRUCT_BODY()
 
