When showing actors and level instance content in the Data Layer outliner, could not interact with level instance part of a data layer.

diff --git a/Engine/Source/Editor/DataLayerEditor/Private/DataLayer/DataLayerMode.cpp b/Engine/Source/Editor/DataLayerEditor/Private/DataLayer/DataLayerMode.cpp
--- a/Engine/Source/Editor/DataLayerEditor/Private/DataLayer/DataLayerMode.cpp
+++ b/Engine/Source/Editor/DataLayerEditor/Private/DataLayer/DataLayerMode.cpp
@@ -234,10 +234,8 @@ FDataLayerMode::FDataLayerMode(const FDataLayerModeParams& Params)
 		{
 			if (const ULevelInstanceSubsystem* LevelInstanceSubsystem = UWorld::GetSubsystem<ULevelInstanceSubsystem>(RepresentingWorld.Get()))
 			{
-				const ILevelInstanceInterface* ActorAsLevelInstance = Cast<ILevelInstanceInterface>(Actor);
-				const ILevelInstanceInterface* ActorParentLevelInstance = LevelInstanceSubsystem->GetParentLevelInstance(Actor);
-				if (!LevelInstanceSubsystem->IsEditingLevelInstance(ActorAsLevelInstance) &&
-					!LevelInstanceSubsystem->IsEditingLevelInstance(ActorParentLevelInstance))
+				const ILevelInstanceInterface* ParentLevelInstance = LevelInstanceSubsystem->GetParentLevelInstance(Actor);
+				if (ParentLevelInstance && !LevelInstanceSubsystem->IsEditingLevelInstance(ParentLevelInstance))
 				{
 					return false;
 				}
