Show child actors of child level instances even when editing a level instance

diff --git a/Engine/Source/Editor/DataLayerEditor/Private/DataLayer/DataLayerActorTreeItem.h b/Engine/Source/Editor/DataLayerEditor/Private/DataLayer/DataLayerActorTreeItem.h
--- a/Engine/Source/Editor/DataLayerEditor/Private/DataLayer/DataLayerActorTreeItem.h
+++ b/Engine/Source/Editor/DataLayerEditor/Private/DataLayer/DataLayerActorTreeItem.h
@@ -86,9 +86,9 @@ private:
 		{
 			ULevel* Level = Actor.IsValid() ? Actor->GetLevel() : nullptr;
 			ULevelInstanceSubsystem* LevelInstanceSubsystem = UWorld::GetSubsystem<ULevelInstanceSubsystem>(OwningWorld);
-			if (LevelInstanceSubsystem && Level && (Level != OwningWorld->GetCurrentLevel()))
+			if (LevelInstanceSubsystem && Level)
 			{
-				DisplayString = LevelInstanceSubsystem->PrefixWithParentLevelInstanceActorLabels(DisplayString, Actor->GetLevel());
+				DisplayString = LevelInstanceSubsystem->PrefixWithParentLevelInstanceActorLabels(DisplayString, Level);
 			}
 		}
 	}
diff --git a/Engine/Source/Editor/DataLayerEditor/Private/DataLayer/DataLayerHierarchy.cpp b/Engine/Source/Editor/DataLayerEditor/Private/DataLayer/DataLayerHierarchy.cpp
--- a/Engine/Source/Editor/DataLayerEditor/Private/DataLayer/DataLayerHierarchy.cpp
+++ b/Engine/Source/Editor/DataLayerEditor/Private/DataLayer/DataLayerHierarchy.cpp
@@ -209,10 +209,36 @@ void FDataLayerHierarchy::CreateItems(TArray<FSceneOutlinerTreeItemPtr>& OutItem
 
 	if (bShowDataLayerActors)
 	{
+		const ULevelInstanceSubsystem* LevelInstanceSubsystem = UWorld::GetSubsystem<ULevelInstanceSubsystem>(GetOwningWorld());
+
+		// Build a set of considered level instances to show nested level instance actors and their child actors
+		TSet<const ILevelInstanceInterface*> ConsideredLevelInstances;		
+		if (bShowLevelInstanceContent && CurrentLevel && LevelInstanceSubsystem)
+		{
+			ILevelInstanceInterface* CurrentLevelInstance = LevelInstanceSubsystem->GetOwningLevelInstance(CurrentLevel);
+			if (LevelInstanceSubsystem->IsEditingLevelInstance(CurrentLevelInstance))
+			{
+				const bool bRecursive = true;
+				LevelInstanceSubsystem->ForEachLevelInstanceChild(CurrentLevelInstance, bRecursive, [&ConsideredLevelInstances](const ILevelInstanceInterface* ChildLevelInstance)
+				{
+					ConsideredLevelInstances.Add(ChildLevelInstance);
+					return true;
+				});
+			}
+		}
+
 		for (AActor* Actor : FActorRange(GetOwningWorld()))
 		{
-			// Only consider actors of current level (if there is one)
-			if (!CurrentLevel || (Actor->GetLevel() == CurrentLevel))
+			// Consider all actors or actors part of current level (if there is one)
+			bool bConsiderActor = !CurrentLevel || (Actor->GetLevel() == CurrentLevel);
+			if (!bConsiderActor && bShowLevelInstanceContent)
+			{
+				// If current level is referenced by editing level instance, also consider child actors part of any child Level Instances of this level
+				ILevelInstanceInterface* ParentLevelInstance = LevelInstanceSubsystem->GetParentLevelInstance(Actor);
+				bConsiderActor = ParentLevelInstance && ConsideredLevelInstances.Contains(ParentLevelInstance);
+			}
+
+			if (bConsiderActor)
 			{
 				for (const UDataLayerInstance* DataLayerInstance : Actor->GetDataLayerInstances())
 				{
@@ -229,7 +255,6 @@ void FDataLayerHierarchy::CreateItems(TArray<FSceneOutlinerTreeItemPtr>& OutItem
 
 		if (bShowUnloadedActors)
 		{
-			ULevelInstanceSubsystem* LevelInstanceSubsystem = UWorld::GetSubsystem<ULevelInstanceSubsystem>(GetOwningWorld());
 			if (UWorldPartitionSubsystem* WorldPartitionSubsystem = UWorld::GetSubsystem<UWorldPartitionSubsystem>(GetOwningWorld()))
 			{
 				WorldPartitionSubsystem->ForEachWorldPartition([this, CurrentLevel, LevelInstanceSubsystem, DataLayerSubsystem, IsDataLayerShown, &OutItems](UWorldPartition* WorldPartition)
diff --git a/Engine/Source/Editor/DataLayerEditor/Private/DataLayer/DataLayerMode.cpp b/Engine/Source/Editor/DataLayerEditor/Private/DataLayer/DataLayerMode.cpp
--- a/Engine/Source/Editor/DataLayerEditor/Private/DataLayer/DataLayerMode.cpp
+++ b/Engine/Source/Editor/DataLayerEditor/Private/DataLayer/DataLayerMode.cpp
@@ -1872,7 +1872,6 @@ void FDataLayerMode::ChooseRepresentingWorld()
 	if (RepresentingWorld == nullptr)
 	{
 		// still not world so fallback to old logic where we just prefer PIE over Editor
-
 		for (const FWorldContext& Context : GEngine->GetWorldContexts())
 		{
 			if (Context.WorldType == EWorldType::PIE)
@@ -1883,16 +1882,6 @@ void FDataLayerMode::ChooseRepresentingWorld()
 			else if (Context.WorldType == EWorldType::Editor)
 			{
 				RepresentingWorld = Context.World();
-				if (RepresentingWorld.IsValid())
-				{
-					// Favor Current Level's outer world if exists
-					ULevel* CurrentLevel = RepresentingWorld->GetCurrentLevel();
-					if (CurrentLevel && !CurrentLevel->IsPersistentLevel() && CurrentLevel->GetWorldDataLayers())
-					{
-						RepresentingWorld = CurrentLevel->GetTypedOuter<UWorld>();
-						check(RepresentingWorld.IsValid());
-					}
-				}
 			}
 		}
 	}
