Fixed static analysis dereferencing null pointer in DataLayerHierarchy.cpp

diff --git a/Engine/Source/Editor/DataLayerEditor/Private/DataLayer/DataLayerHierarchy.cpp b/Engine/Source/Editor/DataLayerEditor/Private/DataLayer/DataLayerHierarchy.cpp
--- a/Engine/Source/Editor/DataLayerEditor/Private/DataLayer/DataLayerHierarchy.cpp
+++ b/Engine/Source/Editor/DataLayerEditor/Private/DataLayer/DataLayerHierarchy.cpp
@@ -231,7 +231,7 @@ void FDataLayerHierarchy::CreateItems(TArray<FSceneOutlinerTreeItemPtr>& OutItem
 		{
 			// Consider all actors or actors part of current level (if there is one)
 			bool bConsiderActor = !CurrentLevel || (Actor->GetLevel() == CurrentLevel);
-			if (!bConsiderActor && bShowLevelInstanceContent)
+			if (!bConsiderActor && bShowLevelInstanceContent && LevelInstanceSubsystem)
 			{
 				// If current level is referenced by editing level instance, also consider child actors part of any child Level Instances of this level
 				ILevelInstanceInterface* ParentLevelInstance = LevelInstanceSubsystem->GetParentLevelInstance(Actor);
