Data Layers: don't allow user changes for non-user managed actor types.

diff --git a/Engine/Source/Editor/DataLayerEditor/Private/DataLayer/DataLayerMode.cpp b/Engine/Source/Editor/DataLayerEditor/Private/DataLayer/DataLayerMode.cpp
--- a/Engine/Source/Editor/DataLayerEditor/Private/DataLayer/DataLayerMode.cpp
+++ b/Engine/Source/Editor/DataLayerEditor/Private/DataLayer/DataLayerMode.cpp
@@ -604,7 +604,7 @@ FSceneOutlinerDragValidationInfo FDataLayerMode::ValidateDrop(const ISceneOutlin
 	{
 		for (AActor* Actor : PayloadActors)
 		{
-			if (!DataLayerEditorSubsystem->IsActorValidForDataLayer(Actor))
+			if (!Actor->IsUserManaged() ||  !DataLayerEditorSubsystem->IsActorValidForDataLayer(Actor))
 			{
 				return FSceneOutlinerDragValidationInfo(ESceneOutlinerDropCompatibility::IncompatibleGeneric, LOCTEXT("ActorCantBeAssignedToDataLayer", "Can't assign actors to Data Layer"));
 			}
diff --git a/Engine/Source/Runtime/Engine/Private/ActorEditor.cpp b/Engine/Source/Runtime/Engine/Private/ActorEditor.cpp
--- a/Engine/Source/Runtime/Engine/Private/ActorEditor.cpp
+++ b/Engine/Source/Runtime/Engine/Private/ActorEditor.cpp
@@ -96,7 +96,7 @@ bool AActor::CanEditChange(const FProperty* PropertyThatWillChange) const
 		return false;
 	}
 
-	if (bIsDataLayersProperty && !SupportsDataLayer())
+	if (bIsDataLayersProperty && (!SupportsDataLayer() || !IsUserManaged()))
 	{
 		return false;
 	}
