Redo CL 23850205:
- Make FSceneOutlinerTreeItemID taking FGuid constructor explicit to find issues where wrong ID was passed in
- Convert DataLayer tree items to use the FArchiveMD5

diff --git a/Engine/Plugins/Editor/ObjectMixer/ObjectMixer/Source/ObjectMixer/Private/Views/List/Modes/ObjectMixerOutlinerHierarchy.cpp b/Engine/Plugins/Editor/ObjectMixer/ObjectMixer/Source/ObjectMixer/Private/Views/List/Modes/ObjectMixerOutlinerHierarchy.cpp
--- a/Engine/Plugins/Editor/ObjectMixer/ObjectMixer/Source/ObjectMixer/Private/Views/List/Modes/ObjectMixerOutlinerHierarchy.cpp
+++ b/Engine/Plugins/Editor/ObjectMixer/ObjectMixer/Source/ObjectMixer/Private/Views/List/Modes/ObjectMixerOutlinerHierarchy.cpp
@@ -364,7 +364,7 @@ FSceneOutlinerTreeItemPtr FObjectMixerOutlinerHierarchy::FindOrCreateParentItem(
 						}
 
 						// Find unloaded parent actor node (from the guid)
-						if (const FSceneOutlinerTreeItemPtr* ParentItem = Items.Find(ParentActorGuid))
+						if (const FSceneOutlinerTreeItemPtr* ParentItem = Items.Find(FActorDescTreeItem::ComputeTreeItemID(ParentActorGuid, ActorDescContainer)))
 						{
 							return *ParentItem;
 						}
@@ -775,10 +775,18 @@ void FObjectMixerOutlinerHierarchy::OnLoadedActorAdded(AActor& InActor)
 {
 	OnLevelActorAdded(&InActor);
 
-	FSceneOutlinerHierarchyChangedData EventData;
-	EventData.Type = FSceneOutlinerHierarchyChangedData::Removed;
-	EventData.ItemIDs.Add(InActor.GetActorGuid());
-	HierarchyChangedEvent.Broadcast(EventData);
+	// Loaded actors can be in sub-world partitions so we use the outer world
+	if (UWorldPartition* WorldPartition = FWorldPartitionHelpers::GetWorldPartition(&InActor))
+	{
+		const FGuid& ActorGuid = InActor.GetActorGuid();
+		if (FWorldPartitionActorDesc* ActorDesc = WorldPartition->GetActorDesc(ActorGuid))
+		{
+			FSceneOutlinerHierarchyChangedData EventData;
+			EventData.Type = FSceneOutlinerHierarchyChangedData::Removed;
+			EventData.ItemIDs.Add(FActorDescTreeItem::ComputeTreeItemID(ActorDesc->GetGuid(), ActorDesc->GetContainer()));
+			HierarchyChangedEvent.Broadcast(EventData);
+		}
+	}
 }
 
 void FObjectMixerOutlinerHierarchy::OnLoadedActorRemoved(AActor& InActor)
@@ -792,7 +800,7 @@ void FObjectMixerOutlinerHierarchy::OnActorDescRemoved(FWorldPartitionActorDesc*
 	{
 		FSceneOutlinerHierarchyChangedData EventData;
 		EventData.Type = FSceneOutlinerHierarchyChangedData::Removed;
-		EventData.ItemIDs.Add(ActorDesc->GetGuid());
+		EventData.ItemIDs.Add(FActorDescTreeItem::ComputeTreeItemID(ActorDesc->GetGuid(), ActorDesc->GetContainer()));
 		HierarchyChangedEvent.Broadcast(EventData);
 	}
 }
diff --git a/Engine/Plugins/Editor/ObjectMixer/ObjectMixer/Source/ObjectMixer/Private/Views/List/Modes/ObjectMixerOutlinerMode.cpp b/Engine/Plugins/Editor/ObjectMixer/ObjectMixer/Source/ObjectMixer/Private/Views/List/Modes/ObjectMixerOutlinerMode.cpp
--- a/Engine/Plugins/Editor/ObjectMixer/ObjectMixer/Source/ObjectMixer/Private/Views/List/Modes/ObjectMixerOutlinerMode.cpp
+++ b/Engine/Plugins/Editor/ObjectMixer/ObjectMixer/Source/ObjectMixer/Private/Views/List/Modes/ObjectMixerOutlinerMode.cpp
@@ -1652,19 +1652,31 @@ void FObjectMixerOutlinerMode::OnLevelActorDeleted(AActor* Actor)
 
 void FObjectMixerOutlinerMode::OnSelectUnloadedActors(const TArray<FGuid>& ActorGuids)
 {
-	TArray<FSceneOutlinerTreeItemPtr> ItemsToSelect;
-	ItemsToSelect.Reserve(ActorGuids.Num());
-	for (const FGuid& ActorGuid : ActorGuids)
+	if (UWorldPartition* WorldPartition = RepresentingWorld->GetWorldPartition())
 	{
-		if (FSceneOutlinerTreeItemPtr ItemPtr = SceneOutliner->GetTreeItem(ActorGuid))
+		TArray<FSceneOutlinerTreeItemPtr> ItemsToSelect;
+		ItemsToSelect.Reserve(ActorGuids.Num());
+		for (const FGuid& ActorGuid : ActorGuids)
 		{
-			ItemsToSelect.Add(ItemPtr);
+			if (FWorldPartitionActorDesc* ActorDesc = WorldPartition->GetActorDesc(ActorGuid))
+			{
+				if (FSceneOutlinerTreeItemPtr ItemPtr = SceneOutliner->GetTreeItem(FActorDescTreeItem::ComputeTreeItemID(ActorDesc->GetGuid(), ActorDesc->GetContainer())))
+				{
+					ItemsToSelect.Add(ItemPtr);
+				}
+			}
 		}
-	}
 
-	if (ItemsToSelect.Num())
-	{
-		SceneOutliner->SetItemSelection(ItemsToSelect, true);
+		if (ItemsToSelect.Num())
+		{
+			SceneOutliner->SetItemSelection(ItemsToSelect, true);
+			SceneOutliner->ScrollItemIntoView(ItemsToSelect.Last());
+
+			if (const FActorDescTreeItem* ActorDescItem = ItemsToSelect.Last()->CastTo<FActorDescTreeItem>())
+			{
+				ActorDescItem->FocusActorBounds();
+			}
+		}
 	}
 }
 
diff --git a/Engine/Source/Editor/DataLayerEditor/Private/DataLayer/DataLayerActorTreeItem.h b/Engine/Source/Editor/DataLayerEditor/Private/DataLayer/DataLayerActorTreeItem.h
--- a/Engine/Source/Editor/DataLayerEditor/Private/DataLayer/DataLayerActorTreeItem.h
+++ b/Engine/Source/Editor/DataLayerEditor/Private/DataLayer/DataLayerActorTreeItem.h
@@ -9,6 +9,9 @@
 #include "Engine/World.h"
 #include "WorldPartition/DataLayer/DataLayerInstance.h"
 #include "LevelInstance/LevelInstanceSubsystem.h"
+#include "Misc/ArchiveMD5.h"
+#include "Misc/SecureHash.h"
+#include "UObject/ObjectKey.h"
 
 //////////////////////////////////////////////////////////////////////////
 // FDataLayerActorTreeItemData
@@ -45,9 +48,19 @@ public:
 	const AActor* GetActor() const { return Actor.Get(); }
 	AActor* GetActor() { return Actor.Get(); }
 
-	static uint32 ComputeTreeItemID(const AActor* InActor, const UDataLayerInstance* InDataLayerInstance)
+	static FSceneOutlinerTreeItemID ComputeTreeItemID(AActor* InActor, const UDataLayerInstance* InDataLayerInstance)
 	{
-		return HashCombine(GetTypeHash(FObjectKey(InActor)), GetTypeHash(FObjectKey(InDataLayerInstance)));
+		FArchiveMD5 Ar;
+		FObjectKey ActorKey(InActor);
+		Ar << ActorKey;
+
+		FObjectKey DataLayerInstanceKey(InDataLayerInstance);
+		Ar << DataLayerInstanceKey;
+
+		FMD5Hash MD5Hash;
+		Ar.GetHash(MD5Hash);
+
+		return FSceneOutlinerTreeItemID(MD5HashToGuid(MD5Hash));
 	}
 
 	bool Filter(FFilterPredicate Pred) const
@@ -94,5 +107,5 @@ private:
 	}
 
 	TWeakObjectPtr<UDataLayerInstance> DataLayerInstance;
-	const uint32 IDDataLayerActor;
+	const FSceneOutlinerTreeItemID IDDataLayerActor;
 };
\ No newline at end of file
diff --git a/Engine/Source/Editor/DataLayerEditor/Private/DataLayer/DataLayerHierarchy.cpp b/Engine/Source/Editor/DataLayerEditor/Private/DataLayer/DataLayerHierarchy.cpp
--- a/Engine/Source/Editor/DataLayerEditor/Private/DataLayer/DataLayerHierarchy.cpp
+++ b/Engine/Source/Editor/DataLayerEditor/Private/DataLayer/DataLayerHierarchy.cpp
@@ -506,13 +506,18 @@ void FDataLayerHierarchy::OnLoadedActorAdded(AActor& InActor)
 			FSceneOutlinerHierarchyChangedData EventData;
 			EventData.Type = FSceneOutlinerHierarchyChangedData::Removed;
 			EventData.ItemIDs.Reserve(DataLayerInstances.Num());
-			ULevelInstanceSubsystem* LevelInstanceSubsystem = UWorld::GetSubsystem<ULevelInstanceSubsystem>(GetOwningWorld());
-			TArray<AActor*> ParentActors = LevelInstanceSubsystem->GetParentLevelInstanceActors(InActor.GetLevel());
-
-			for (const UDataLayerInstance* DataLayerInstance : DataLayerInstances)
+		
+			if (UWorldPartition* WorldPartition = FWorldPartitionHelpers::GetWorldPartition(&InActor))
 			{
-				EventData.ItemIDs.Add(FDataLayerActorDescTreeItem::ComputeTreeItemID(InActor.GetActorGuid(), DataLayerInstance, ParentActors));
+				if (FWorldPartitionActorDesc* ActorDesc = WorldPartition->GetActorDesc(InActor.GetActorGuid()))
+				{
+					for (const UDataLayerInstance* DataLayerInstance : DataLayerInstances)
+					{
+						EventData.ItemIDs.Add(FDataLayerActorDescTreeItem::ComputeTreeItemID(ActorDesc->GetGuid(), ActorDesc->GetContainer(), DataLayerInstance));
+					}
+				}
 			}
+
 			HierarchyChangedEvent.Broadcast(EventData);
 		}
 	}
@@ -524,7 +529,7 @@ void FDataLayerHierarchy::OnLoadedActorRemoved(AActor& InActor)
 	OnLevelActorDeleted(&InActor);
 
 	// Add any corresponding actor desc items for this actor
-	UWorldPartition* const WorldPartition = RepresentingWorld->GetWorldPartition();
+	UWorldPartition* const WorldPartition = FWorldPartitionHelpers::GetWorldPartition(&InActor);
 	if ((GetOwningWorld() == InActor.GetWorld()) && WorldPartition != nullptr)
 	{
 		const TArray<const UDataLayerInstance*> DataLayerInstances = InActor.GetDataLayerInstances();
@@ -535,7 +540,10 @@ void FDataLayerHierarchy::OnLoadedActorRemoved(AActor& InActor)
 			EventData.Items.Reserve(DataLayerInstances.Num());
 			for (const UDataLayerInstance* DataLayerInstance : DataLayerInstances)
 			{
-				EventData.Items.Add(Mode->CreateItemFor<FDataLayerActorDescTreeItem>(FDataLayerActorDescTreeItemData(InActor.GetActorGuid(), WorldPartition, const_cast<UDataLayerInstance*>(DataLayerInstance))));
+				if (FWorldPartitionActorDesc* ActorDesc = WorldPartition->GetActorDesc(InActor.GetActorGuid()))
+				{
+					EventData.Items.Add(Mode->CreateItemFor<FDataLayerActorDescTreeItem>(FDataLayerActorDescTreeItemData(ActorDesc->GetGuid(), ActorDesc->GetContainer(), const_cast<UDataLayerInstance*>(DataLayerInstance))));
+				}
 			}
 			HierarchyChangedEvent.Broadcast(EventData);
 		}
@@ -549,7 +557,6 @@ void FDataLayerHierarchy::OnActorDescAdded(FWorldPartitionActorDesc* InActorDesc
 		return;
 	}
 
-	UWorldPartition* const WorldPartition = RepresentingWorld->GetWorldPartition();
 	const UDataLayerSubsystem* const DataLayerSubsystem = UWorld::GetSubsystem<UDataLayerSubsystem>(GetOwningWorld());
 	const TArray<FName>& DataLayerInstanceNames = InActorDesc->GetDataLayerInstanceNames();
 
@@ -562,7 +569,7 @@ void FDataLayerHierarchy::OnActorDescAdded(FWorldPartitionActorDesc* InActorDesc
 		for (const FName& DataLayerInstanceName : DataLayerInstanceNames)
 		{
 			const UDataLayerInstance* const DataLayerInstance = DataLayerSubsystem->GetDataLayerInstance(DataLayerInstanceName);
-			EventData.Items.Add(Mode->CreateItemFor<FDataLayerActorDescTreeItem>(FDataLayerActorDescTreeItemData(InActorDesc->GetGuid(), WorldPartition, const_cast<UDataLayerInstance*>(DataLayerInstance))));
+			EventData.Items.Add(Mode->CreateItemFor<FDataLayerActorDescTreeItem>(FDataLayerActorDescTreeItemData(InActorDesc->GetGuid(), InActorDesc->GetContainer(), const_cast<UDataLayerInstance*>(DataLayerInstance))));
 		}
 		HierarchyChangedEvent.Broadcast(EventData);
 	}
@@ -583,12 +590,11 @@ void FDataLayerHierarchy::OnActorDescRemoved(FWorldPartitionActorDesc* InActorDe
 		FSceneOutlinerHierarchyChangedData EventData;
 		EventData.Type = FSceneOutlinerHierarchyChangedData::Removed;
 		EventData.ItemIDs.Reserve(DataLayerInstanceNames.Num());
-		TArray<AActor*> ParentActors = FDataLayerActorDescTreeItem::GetParentActors(InActorDesc->GetContainer());
 
 		for (const FName& DataLayerInstanceName : DataLayerInstanceNames)
 		{
 			const UDataLayerInstance* const DataLayer = DataLayerSubsystem->GetDataLayerInstance(DataLayerInstanceName);
-			EventData.ItemIDs.Add(FDataLayerActorDescTreeItem::ComputeTreeItemID(InActorDesc->GetGuid(), DataLayer, ParentActors));
+			EventData.ItemIDs.Add(FDataLayerActorDescTreeItem::ComputeTreeItemID(InActorDesc->GetGuid(), InActorDesc->GetContainer(), DataLayer));
 		}
 		HierarchyChangedEvent.Broadcast(EventData);
 	}
diff --git a/Engine/Source/Editor/DataLayerEditor/Private/DataLayer/DataLayersActorDescTreeItem.h b/Engine/Source/Editor/DataLayerEditor/Private/DataLayer/DataLayersActorDescTreeItem.h
--- a/Engine/Source/Editor/DataLayerEditor/Private/DataLayer/DataLayersActorDescTreeItem.h
+++ b/Engine/Source/Editor/DataLayerEditor/Private/DataLayer/DataLayersActorDescTreeItem.h
@@ -9,6 +9,10 @@
 #include "WorldPartition/DataLayer/DataLayerInstance.h"
 #include "WorldPartition/ActorDescContainerCollection.h"
 #include "LevelInstance/LevelInstanceSubsystem.h"
+#include "Misc/ArchiveMD5.h"
+#include "Misc/SecureHash.h"
+#include "UObject/ObjectKey.h"
+
 
 //////////////////////////////////////////////////////////////////////////
 // FDataLayerActorTreeItemData
@@ -23,12 +27,6 @@ struct FDataLayerActorDescTreeItemData
 		, DataLayer(InDataLayer)
 	{}
 
-	FDataLayerActorDescTreeItemData(const FGuid& InActorGuid, FActorDescContainerCollection* InContainerCollection, UDataLayerInstance* InDataLayer)
-		: ActorGuid(InActorGuid)
-		, Container(InContainerCollection->GetActorDescContainer(InActorGuid))
-		, DataLayer(InDataLayer)
-	{}
-	
 	const FGuid& ActorGuid;
 	UActorDescContainer* const Container;
 	TWeakObjectPtr<UDataLayerInstance> DataLayer;
@@ -46,7 +44,7 @@ public:
 	FDataLayerActorDescTreeItem(const FDataLayerActorDescTreeItemData& InData)
 		: FActorDescTreeItem(InData.ActorGuid, InData.Container)
 		, DataLayer(InData.DataLayer)
-		, IDDataLayerActorDesc(FDataLayerActorDescTreeItem::ComputeTreeItemID(InData.ActorGuid, DataLayer.Get(), FDataLayerActorDescTreeItem::GetParentActors(InData.Container)))
+		, IDDataLayerActorDesc(FDataLayerActorDescTreeItem::ComputeTreeItemID(InData.ActorGuid, InData.Container, InData.DataLayer.Get()))
 	{
 		if (ActorDescHandle.IsValid())
 		{
@@ -54,7 +52,7 @@ public:
 			UWorld* OwningWorld = Container->GetWorld();
 			ULevelInstanceSubsystem* LevelInstanceSubsystem = UWorld::GetSubsystem<ULevelInstanceSubsystem>(OwningWorld);
 			ULevel* Level = Container->GetTypedOuter<UWorld>()->PersistentLevel;
-			if (LevelInstanceSubsystem && (Level != OwningWorld->GetCurrentLevel()))
+			if (LevelInstanceSubsystem && Level)
 			{
 				DisplayString = LevelInstanceSubsystem->PrefixWithParentLevelInstanceActorLabels(DisplayString, Level);
 			}
@@ -63,14 +61,21 @@ public:
 
 	UDataLayerInstance* GetDataLayer() const { return DataLayer.Get(); }
 	
-	static uint32 ComputeTreeItemID(const FGuid& InActorGuid, const UDataLayerInstance* InDataLayer, const TArray<AActor*>& InParentActors)
+	static FSceneOutlinerTreeItemID ComputeTreeItemID(FGuid InActorGuid, UActorDescContainer* InContainer, const UDataLayerInstance* InDataLayer)
 	{
-		uint32 ID = HashCombine(GetTypeHash(InActorGuid), GetTypeHash(FObjectKey(InDataLayer)));
-		for (AActor* ParentActor : InParentActors)
-		{
-			ID = HashCombine(ID, GetTypeHash(ParentActor->GetActorGuid()));
-		}
-		return ID;
+		FArchiveMD5 Ar;
+		Ar << InActorGuid;
+
+		FObjectKey ContainerKey(InContainer);
+		Ar << ContainerKey;
+
+		FObjectKey DataLayerInstanceKey(InDataLayer);
+		Ar << DataLayerInstanceKey;
+
+		FMD5Hash MD5Hash;
+		Ar.GetHash(MD5Hash);
+
+		return FSceneOutlinerTreeItemID(MD5HashToGuid(MD5Hash));
 	}
 
 	static TArray<AActor*> GetParentActors(UActorDescContainer* InContainer)
@@ -107,5 +112,5 @@ public:
 
 private:
 	TWeakObjectPtr<UDataLayerInstance> DataLayer;
-	const uint32 IDDataLayerActorDesc;
+	const FSceneOutlinerTreeItemID IDDataLayerActorDesc;
 };
\ No newline at end of file
diff --git a/Engine/Source/Editor/SceneOutliner/Private/ActorBrowsingMode.cpp b/Engine/Source/Editor/SceneOutliner/Private/ActorBrowsingMode.cpp
--- a/Engine/Source/Editor/SceneOutliner/Private/ActorBrowsingMode.cpp
+++ b/Engine/Source/Editor/SceneOutliner/Private/ActorBrowsingMode.cpp
@@ -918,24 +918,30 @@ void FActorBrowsingMode::OnLevelActorDeleted(AActor* Actor)
 
 void FActorBrowsingMode::OnSelectUnloadedActors(const TArray<FGuid>& ActorGuids)
 {
-	TArray<FSceneOutlinerTreeItemPtr> ItemsToSelect;
-	ItemsToSelect.Reserve(ActorGuids.Num());
-	for (const FGuid& ActorGuid : ActorGuids)
+	if (UWorldPartition* WorldPartition = RepresentingWorld->GetWorldPartition())
 	{
-		if (FSceneOutlinerTreeItemPtr ItemPtr = SceneOutliner->GetTreeItem(ActorGuid))
+		TArray<FSceneOutlinerTreeItemPtr> ItemsToSelect;
+		ItemsToSelect.Reserve(ActorGuids.Num());
+		for (const FGuid& ActorGuid : ActorGuids)
 		{
-			ItemsToSelect.Add(ItemPtr);
+			if (FWorldPartitionActorDesc* ActorDesc = WorldPartition->GetActorDesc(ActorGuid))
+			{
+				if (FSceneOutlinerTreeItemPtr ItemPtr = SceneOutliner->GetTreeItem(FActorDescTreeItem::ComputeTreeItemID(ActorDesc->GetGuid(), ActorDesc->GetContainer())))
+				{
+					ItemsToSelect.Add(ItemPtr);
+				}
+			}
 		}
-	}
 
-	if (ItemsToSelect.Num())
-	{
-		SceneOutliner->SetItemSelection(ItemsToSelect, true);
-		SceneOutliner->ScrollItemIntoView(ItemsToSelect.Last());
-
-		if (const FActorDescTreeItem* ActorDescItem = ItemsToSelect.Last()->CastTo<FActorDescTreeItem>())
+		if (ItemsToSelect.Num())
 		{
-			ActorDescItem->FocusActorBounds();
+			SceneOutliner->SetItemSelection(ItemsToSelect, true);
+			SceneOutliner->ScrollItemIntoView(ItemsToSelect.Last());
+
+			if (const FActorDescTreeItem* ActorDescItem = ItemsToSelect.Last()->CastTo<FActorDescTreeItem>())
+			{
+				ActorDescItem->FocusActorBounds();
+			}
 		}
 	}
 }
diff --git a/Engine/Source/Editor/SceneOutliner/Private/ActorDescTreeItem.cpp b/Engine/Source/Editor/SceneOutliner/Private/ActorDescTreeItem.cpp
--- a/Engine/Source/Editor/SceneOutliner/Private/ActorDescTreeItem.cpp
+++ b/Engine/Source/Editor/SceneOutliner/Private/ActorDescTreeItem.cpp
@@ -24,6 +24,10 @@
 #include "ToolMenus.h"
 #include "LevelEditorViewport.h"
 #include "HAL/PlatformApplicationMisc.h"
+#include "Misc/ArchiveMD5.h"
+#include "Misc/SecureHash.h"
+#include "UObject/ObjectKey.h"
+
 
 #define LOCTEXT_NAMESPACE "SceneOutliner_ActorDescTreeItem"
 
@@ -223,26 +227,12 @@ private:
 	}
 };
 
-FActorDescTreeItem::FActorDescTreeItem(const FGuid& InActorGuid, UActorDescContainer* Container)
-	: IActorBaseTreeItem(Type)
-	, ActorDescHandle(Container, InActorGuid)
-	, ActorGuid(InActorGuid)
-{
-	Initialize();
-}
-
-FActorDescTreeItem::FActorDescTreeItem(const FGuid& InActorGuid, FActorDescContainerCollection* ContainerCollection)
+FActorDescTreeItem::FActorDescTreeItem(const FGuid& InActorGuid, UActorDescContainer* InContainer)
 	: IActorBaseTreeItem(Type)
-	, ActorDescHandle(ContainerCollection, InActorGuid)
+	, ActorDescHandle(InContainer, InActorGuid)
+	, ID(ComputeTreeItemID(InActorGuid, InContainer))
 	, ActorGuid(InActorGuid)
 {
-	Initialize();
-}
-
-void FActorDescTreeItem::Initialize()
-{
-	ID = FSceneOutlinerTreeItemID(ActorGuid);
-
 	if (const FWorldPartitionActorDesc* const ActorDesc = ActorDescHandle.Get())
 	{
 		DisplayString = ActorDesc->GetActorLabel().ToString();
@@ -253,6 +243,20 @@ void FActorDescTreeItem::Initialize()
 	}
 }
 
+FSceneOutlinerTreeItemID FActorDescTreeItem::ComputeTreeItemID(FGuid InActorGuid, UActorDescContainer* InContainer)
+{
+	FArchiveMD5 Ar;
+	Ar << InActorGuid;
+
+	FObjectKey ContainerKey(InContainer);
+	Ar << ContainerKey;
+
+	FMD5Hash MD5Hash;
+	Ar.GetHash(MD5Hash);
+	
+	return FSceneOutlinerTreeItemID(MD5HashToGuid(MD5Hash));
+}
+
 FSceneOutlinerTreeItemID FActorDescTreeItem::GetID() const
 {
 	return ID;
diff --git a/Engine/Source/Editor/SceneOutliner/Private/ActorHierarchy.cpp b/Engine/Source/Editor/SceneOutliner/Private/ActorHierarchy.cpp
--- a/Engine/Source/Editor/SceneOutliner/Private/ActorHierarchy.cpp
+++ b/Engine/Source/Editor/SceneOutliner/Private/ActorHierarchy.cpp
@@ -312,7 +312,7 @@ FSceneOutlinerTreeItemPtr FActorHierarchy::FindOrCreateParentItem(const ISceneOu
 						}
 
 						// Find unloaded parent actor node (from the guid)
-						if (const FSceneOutlinerTreeItemPtr* ParentItem = Items.Find(ParentActorGuid))
+						if (const FSceneOutlinerTreeItemPtr* ParentItem = Items.Find(FActorDescTreeItem::ComputeTreeItemID(ParentActorGuid, ActorDescContainer)))
 						{
 							return *ParentItem;
 						}
@@ -638,10 +638,18 @@ void FActorHierarchy::OnLoadedActorAdded(AActor& InActor)
 {
 	OnLevelActorAdded(&InActor);
 
-	FSceneOutlinerHierarchyChangedData EventData;
-	EventData.Type = FSceneOutlinerHierarchyChangedData::Removed;
-	EventData.ItemIDs.Add(InActor.GetActorGuid());
-	HierarchyChangedEvent.Broadcast(EventData);
+	// Loaded actors can be in sub-world partitions so we use the outer world
+	if (UWorldPartition* WorldPartition = FWorldPartitionHelpers::GetWorldPartition(&InActor))
+	{
+		const FGuid& ActorGuid = InActor.GetActorGuid();
+		if (FWorldPartitionActorDesc* ActorDesc = WorldPartition->GetActorDesc(ActorGuid))
+		{
+			FSceneOutlinerHierarchyChangedData EventData;
+			EventData.Type = FSceneOutlinerHierarchyChangedData::Removed;
+			EventData.ItemIDs.Add(FActorDescTreeItem::ComputeTreeItemID(ActorGuid, ActorDesc->GetContainer()));
+			HierarchyChangedEvent.Broadcast(EventData);
+		}
+	}
 }
 
 void FActorHierarchy::OnLoadedActorRemoved(AActor& InActor)
@@ -650,14 +658,14 @@ void FActorHierarchy::OnLoadedActorRemoved(AActor& InActor)
 
 	if (bShowingUnloadedActors)
 	{
-		if (UWorldPartition* WorldPartition = RepresentingWorld->GetWorldPartition())
+		if (UWorldPartition* WorldPartition = FWorldPartitionHelpers::GetWorldPartition(&InActor))
 		{
 			const FGuid& ActorGuid = InActor.GetActorGuid();
-			if (WorldPartition->GetActorDesc(ActorGuid) != nullptr)
+			if (FWorldPartitionActorDesc* ActorDesc = WorldPartition->GetActorDesc(ActorGuid))
 			{
 				FSceneOutlinerHierarchyChangedData EventData;
 				EventData.Type = FSceneOutlinerHierarchyChangedData::Added;
-				EventData.Items.Add(Mode->CreateItemFor<FActorDescTreeItem>(FActorDescTreeItem(ActorGuid, WorldPartition)));
+				EventData.Items.Add(Mode->CreateItemFor<FActorDescTreeItem>(FActorDescTreeItem(ActorDesc->GetGuid(), ActorDesc->GetContainer())));
 				HierarchyChangedEvent.Broadcast(EventData);
 			}
 		}
@@ -668,13 +676,10 @@ void FActorHierarchy::OnActorDescAdded(FWorldPartitionActorDesc* ActorDesc)
 {
 	if (bShowingUnloadedActors && ActorDesc && !ActorDesc->IsLoaded(true))
 	{
-		if (UWorldPartition* WorldPartition = RepresentingWorld->GetWorldPartition())
-		{
-			FSceneOutlinerHierarchyChangedData EventData;
-			EventData.Type = FSceneOutlinerHierarchyChangedData::Added;
-			EventData.Items.Add(Mode->CreateItemFor<FActorDescTreeItem>(FActorDescTreeItem(ActorDesc->GetGuid(), WorldPartition)));
-			HierarchyChangedEvent.Broadcast(EventData);
-		}
+		FSceneOutlinerHierarchyChangedData EventData;
+		EventData.Type = FSceneOutlinerHierarchyChangedData::Added;
+		EventData.Items.Add(Mode->CreateItemFor<FActorDescTreeItem>(FActorDescTreeItem(ActorDesc->GetGuid(), ActorDesc->GetContainer())));
+		HierarchyChangedEvent.Broadcast(EventData);
 	}
 }
 
@@ -684,7 +689,7 @@ void FActorHierarchy::OnActorDescRemoved(FWorldPartitionActorDesc* ActorDesc)
 	{
 		FSceneOutlinerHierarchyChangedData EventData;
 		EventData.Type = FSceneOutlinerHierarchyChangedData::Removed;
-		EventData.ItemIDs.Add(ActorDesc->GetGuid());
+		EventData.ItemIDs.Add(FActorDescTreeItem::ComputeTreeItemID(ActorDesc->GetGuid(), ActorDesc->GetContainer()));
 		HierarchyChangedEvent.Broadcast(EventData);
 	}
 }
diff --git a/Engine/Source/Editor/SceneOutliner/Public/ActorDescTreeItem.h b/Engine/Source/Editor/SceneOutliner/Public/ActorDescTreeItem.h
--- a/Engine/Source/Editor/SceneOutliner/Public/ActorDescTreeItem.h
+++ b/Engine/Source/Editor/SceneOutliner/Public/ActorDescTreeItem.h
@@ -35,14 +35,13 @@ public:
 	/** Constant identifier for this tree item */
 	FSceneOutlinerTreeItemID ID;
 
+	static FSceneOutlinerTreeItemID ComputeTreeItemID(FGuid InActorGuid, UActorDescContainer* InContainer);
+
 	/** Static type identifier for this tree item class */
 	static const FSceneOutlinerTreeItemType Type;
 
 	/** Construct this item from an actor desc */
-	FActorDescTreeItem(const FGuid& InActorGuid, UActorDescContainer* Container);
-
-	/** Construct this item from an actor desc */
-	FActorDescTreeItem(const FGuid& InActorGuid, FActorDescContainerCollection* ContainerCollection);
+	FActorDescTreeItem(const FGuid& InActorGuid, UActorDescContainer* InContainer);
 
 	/* Begin ISceneOutlinerTreeItem Implementation */
 	virtual bool IsValid() const override { return ActorDescHandle.Get() != nullptr; }
@@ -70,6 +69,5 @@ protected:
 
 private:
 	void CopyActorFilePathtoClipboard() const;
-	void Initialize();
 	FGuid ActorGuid;
 };
diff --git a/Engine/Source/Editor/SceneOutliner/Public/SceneOutlinerStandaloneTypes.h b/Engine/Source/Editor/SceneOutliner/Public/SceneOutlinerStandaloneTypes.h
--- a/Engine/Source/Editor/SceneOutliner/Public/SceneOutlinerStandaloneTypes.h
+++ b/Engine/Source/Editor/SceneOutliner/Public/SceneOutlinerStandaloneTypes.h
@@ -72,7 +72,7 @@ public:
 		CachedHash = CalculateTypeHash();
 	}
 
-	FSceneOutlinerTreeItemID(const FGuid& InGuid) : Type(EType::Guid)
+	explicit FSceneOutlinerTreeItemID(const FGuid& InGuid) : Type(EType::Guid)
 	{
 		new (Data) FGuid(InGuid);
 		CachedHash = CalculateTypeHash();
diff --git a/Engine/Source/Editor/WorldPartitionEditor/Private/WorldPartition/ContentBundle/Outliner/SContentBundleOutliner.cpp b/Engine/Source/Editor/WorldPartitionEditor/Private/WorldPartition/ContentBundle/Outliner/SContentBundleOutliner.cpp
--- a/Engine/Source/Editor/WorldPartitionEditor/Private/WorldPartition/ContentBundle/Outliner/SContentBundleOutliner.cpp
+++ b/Engine/Source/Editor/WorldPartitionEditor/Private/WorldPartition/ContentBundle/Outliner/SContentBundleOutliner.cpp
@@ -9,7 +9,7 @@ void SContentBundleOutliner::SelectContentBundle(const TWeakPtr<FContentBundleEd
 	TSharedPtr<FContentBundleEditor> ContentBundleEditorPin = ContentBundle.Pin();
 	if (ContentBundleEditorPin != nullptr)
 	{
-		FSceneOutlinerTreeItemPtr Item = GetTreeItem(ContentBundleEditorPin->GetTreeItemID());
+		FSceneOutlinerTreeItemPtr Item = GetTreeItem(FSceneOutlinerTreeItemID(ContentBundleEditorPin->GetTreeItemID()));
 		if (Item->IsValid())
 		{
 			SetItemSelection(Item, true, ESelectInfo::Direct);
