- Added UDataLayerManager which is local to each WorldPartition.
- This will allow instanced world partition maps to maintain their own Data Layer runtime states.
- Deprecated DataLayerSubsystem since DataLayerManager superseeds its functionality.
- Added helpers like UWorld::GetDataLayerManager and UWorldPartitionBlueprintLibrary::GetDataLayerManager which will automatically find the associated DataLayerManager for the object it's running the BP.

diff --git a/Engine/Source/Editor/DataLayerEditor/Private/DataLayer/DataLayerEditorSubsystem.cpp b/Engine/Source/Editor/DataLayerEditor/Private/DataLayer/DataLayerEditorSubsystem.cpp
--- a/Engine/Source/Editor/DataLayerEditor/Private/DataLayer/DataLayerEditorSubsystem.cpp
+++ b/Engine/Source/Editor/DataLayerEditor/Private/DataLayer/DataLayerEditorSubsystem.cpp
@@ -45,7 +45,7 @@
 #include "Widgets/Text/STextBlock.h"
 #include "WorldPartition/DataLayer/DataLayerAsset.h"
 #include "WorldPartition/DataLayer/DataLayerInstanceWithAsset.h"
-#include "WorldPartition/DataLayer/DataLayerSubsystem.h"
+#include "WorldPartition/DataLayer/DataLayerManager.h"
 #include "WorldPartition/DataLayer/DeprecatedDataLayerInstance.h"
 #include "WorldPartition/DataLayer/WorldDataLayers.h"
 #include "WorldPartition/WorldPartition.h"
@@ -80,8 +80,6 @@ private:
 	void OnObjectPostEditChange(UObject* Object, FPropertyChangedEvent& PropertyChangedEvent);
 	void OnLevelActorsAdded(AActor* InActor) { DataLayerEditorSubsystem->InitializeNewActorDataLayers(InActor); }
 	void OnLevelSelectionChanged(UObject* InObject) { DataLayerEditorSubsystem->OnSelectionChanged(); }
-	void OnWorldDataLayersPostRegister(AWorldDataLayers* WorldDatalayers) { DataLayerEditorSubsystem->OnWorldDataLayersPostRegister(WorldDatalayers); }
-	void OnWorldDataLayersPreUnregister(AWorldDataLayers* WorldDatalayers) { DataLayerEditorSubsystem->OnWorldDataLayersPreUnregister(WorldDatalayers); }
 
 	UDataLayerEditorSubsystem* DataLayerEditorSubsystem;
 	bool bIsInitialized;
@@ -116,8 +114,6 @@ void FDataLayersBroadcast::Deinitialize()
 			}
 			USelection::SelectionChangedEvent.RemoveAll(this);
 			USelection::SelectObjectEvent.RemoveAll(this);
-			UDataLayerSubsystem::OnWorldDataLayerPostRegister.RemoveAll(this);
-			UDataLayerSubsystem::OnWorldDataLayerPreUnregister.RemoveAll(this);
 		}
 	}
 }
@@ -133,14 +129,16 @@ void FDataLayersBroadcast::Initialize()
 		GEngine->OnLevelActorAdded().AddRaw(this, &FDataLayersBroadcast::OnLevelActorsAdded);
 		USelection::SelectionChangedEvent.AddRaw(this, &FDataLayersBroadcast::OnLevelSelectionChanged);
 		USelection::SelectObjectEvent.AddRaw(this, &FDataLayersBroadcast::OnLevelSelectionChanged);
-		UDataLayerSubsystem::OnWorldDataLayerPostRegister.AddRaw(this, &FDataLayersBroadcast::OnWorldDataLayersPostRegister);
-		UDataLayerSubsystem::OnWorldDataLayerPreUnregister.AddRaw(this, &FDataLayersBroadcast::OnWorldDataLayersPreUnregister);
 	}
 }
 
 void FDataLayersBroadcast::OnObjectPostEditChange(UObject* Object, FPropertyChangedEvent& PropertyChangedEvent)
 {
-	if (Object && (PropertyChangedEvent.ChangeType != EPropertyChangeType::Interactive))
+	// Ignore changed on game world objects
+	UWorld* World = Object ? Object->GetWorld() : nullptr;
+	const bool bIsGameWorld = World && World->IsGameWorld();
+
+	if (Object && !bIsGameWorld && (PropertyChangedEvent.ChangeType != EPropertyChangeType::Interactive))
 	{
 		bool bRefresh = false;
 		if (Object->IsA<UDataLayerInstance>() || Object->IsA<UDataLayerAsset>())
@@ -259,28 +257,28 @@ void UDataLayerEditorSubsystem::BeginDestroy()
 void UDataLayerEditorSubsystem::OnExecuteActorEditorContextAction(UWorld* InWorld, const EActorEditorContextAction& InType, AActor* InActor)
 {
 	UE_CLOG(!InWorld, LogDataLayerEditorSubsystem, Error, TEXT("%s - Failed because world in null."), ANSI_TO_TCHAR(__FUNCTION__));
-	if (UDataLayerSubsystem* DataLayerSubsystem = UWorld::GetSubsystem<UDataLayerSubsystem>(InWorld))
+	if (UDataLayerManager* DataLayerManager = UDataLayerManager::GetDataLayerManager(InWorld))
 	{
 		switch (InType)
 		{
 		case EActorEditorContextAction::ApplyContext:
 			check(InActor && InActor->GetWorld() == InWorld);
 			{
-				AddActorToDataLayers(InActor, DataLayerSubsystem->GetActorEditorContextDataLayers());
+				AddActorToDataLayers(InActor, DataLayerManager->GetActorEditorContextDataLayers());
 			}
 			break;
 		case EActorEditorContextAction::ResetContext:
-			for (UDataLayerInstance* DataLayer : DataLayerSubsystem->GetActorEditorContextDataLayers())
+			for (UDataLayerInstance* DataLayerInstance : DataLayerManager->GetActorEditorContextDataLayers())
 			{
-				RemoveFromActorEditorContext(DataLayer);
+				RemoveFromActorEditorContext(DataLayerInstance);
 			}
 			break;
 		case EActorEditorContextAction::PushContext:
-			DataLayerSubsystem->PushActorEditorContext();
+			DataLayerManager->PushActorEditorContext();
 			BroadcastDataLayerChanged(EDataLayerAction::Reset, NULL, NAME_None);
 			break;
 		case EActorEditorContextAction::PopContext:
-			DataLayerSubsystem->PopActorEditorContext();
+			DataLayerManager->PopActorEditorContext();
 			BroadcastDataLayerChanged(EDataLayerAction::Reset, NULL, NAME_None);
 			break;
 		}
@@ -289,9 +287,9 @@ void UDataLayerEditorSubsystem::OnExecuteActorEditorContextAction(UWorld* InWorl
 
 bool UDataLayerEditorSubsystem::GetActorEditorContextDisplayInfo(UWorld* InWorld, FActorEditorContextClientDisplayInfo& OutDiplayInfo) const
 {
-	if (UDataLayerSubsystem* DataLayerSubsystem = UWorld::GetSubsystem<UDataLayerSubsystem>(InWorld))
+	if (UDataLayerManager* DataLayerManager = UDataLayerManager::GetDataLayerManager(InWorld))
 	{
-		if (!DataLayerSubsystem->GetActorEditorContextDataLayers().IsEmpty())
+		if (!DataLayerManager->GetActorEditorContextDataLayers().IsEmpty())
 		{
 			OutDiplayInfo.Title = TEXT("Data Layers");
 			OutDiplayInfo.Brush = FAppStyle::GetBrush(TEXT("DataLayer.Editor"));
@@ -305,12 +303,12 @@ TSharedRef<SWidget> UDataLayerEditorSubsystem::GetActorEditorContextWidget(UWorl
 {
 	TSharedRef<SVerticalBox> OutWidget = SNew(SVerticalBox);
 
-	if (UDataLayerSubsystem* DataLayerSubsystem = UWorld::GetSubsystem<UDataLayerSubsystem>(InWorld))
+	if (UDataLayerManager* DataLayerManager = UDataLayerManager::GetDataLayerManager(InWorld))
 	{
-		TArray<UDataLayerInstance*> DataLayers = DataLayerSubsystem->GetActorEditorContextDataLayers();
-		for (UDataLayerInstance* DataLayer : DataLayers)
+		TArray<UDataLayerInstance*> DataLayers = DataLayerManager->GetActorEditorContextDataLayers();
+		for (UDataLayerInstance* DataLayerInstance : DataLayers)
 		{
-			check(IsValid(DataLayer));
+			check(IsValid(DataLayerInstance));
 			OutWidget->AddSlot().AutoHeight()
 			[
 				SNew(SHorizontalBox)
@@ -320,7 +318,7 @@ TSharedRef<SWidget> UDataLayerEditorSubsystem::GetActorEditorContextWidget(UWorl
 				.VAlign(VAlign_Center)
 				[
 					SNew(SImage)
-					.ColorAndOpacity(DataLayer->GetDebugColor())
+					.ColorAndOpacity(DataLayerInstance->GetDebugColor())
 					.Image(FAppStyle::Get().GetBrush("DataLayer.ColorIcon"))
 					.DesiredSizeOverride(FVector2D(8, 8))
 				]
@@ -329,7 +327,7 @@ TSharedRef<SWidget> UDataLayerEditorSubsystem::GetActorEditorContextWidget(UWorl
 				.Padding(4.0f, 1.0f, 1.0f, 1.0f)
 				[
 					SNew(STextBlock)
-					.Text(FText::FromString(DataLayer->GetDataLayerShortName()))
+					.Text(FText::FromString(DataLayerInstance->GetDataLayerShortName()))
 				]
 			];
 		}
@@ -354,21 +352,11 @@ void UDataLayerEditorSubsystem::RemoveFromActorEditorContext(UDataLayerInstance*
 	}
 }
 
-void UDataLayerEditorSubsystem::OnWorldDataLayersPostRegister(AWorldDataLayers* WorldDataLayers)
-{
-	EditorRefreshDataLayerBrowser();
-}
-
-void UDataLayerEditorSubsystem::OnWorldDataLayersPreUnregister(AWorldDataLayers* WorldDataLayers)
-{
-	EditorRefreshDataLayerBrowser();
-}
-
 TArray<const UDataLayerInstance*> UDataLayerEditorSubsystem::GetDataLayerInstances(const TArray<const UDataLayerAsset*> DataLayerAssets) const
 {
-	if (UDataLayerSubsystem* DataLayerSubsystem = UWorld::GetSubsystem<UDataLayerSubsystem>(GetWorld()))
+	if (UDataLayerManager* DataLayerManager = UDataLayerManager::GetDataLayerManager(GetWorld()))
 	{
-		return DataLayerSubsystem->GetDataLayerInstances(DataLayerAssets);
+		return DataLayerManager->GetDataLayerInstances(DataLayerAssets);
 	}
 
 	return TArray<const UDataLayerInstance*>();
@@ -446,15 +434,15 @@ UWorld* UDataLayerEditorSubsystem::GetWorld() const
 	return GEditor->GetEditorWorldContext().World();
 }
 
-bool UDataLayerEditorSubsystem::SetParentDataLayer(UDataLayerInstance* DataLayer, UDataLayerInstance* ParentDataLayer)
+bool UDataLayerEditorSubsystem::SetParentDataLayer(UDataLayerInstance* DataLayerInstance, UDataLayerInstance* ParentDataLayer)
 {
-	if (DataLayer->CanParent(ParentDataLayer))
+	if (DataLayerInstance->CanParent(ParentDataLayer))
 	{
-		const bool bIsLoaded = DataLayer->IsEffectiveLoadedInEditor();
-		DataLayer->SetParent(ParentDataLayer);
+		const bool bIsLoaded = DataLayerInstance->IsEffectiveLoadedInEditor();
+		DataLayerInstance->SetParent(ParentDataLayer);
 		BroadcastDataLayerChanged(EDataLayerAction::Reset, NULL, NAME_None);
 		UpdateAllActorsVisibility(true, true);
-		if (bIsLoaded != DataLayer->IsEffectiveLoadedInEditor())
+		if (bIsLoaded != DataLayerInstance->IsEffectiveLoadedInEditor())
 		{
 			OnDataLayerEditorLoadingStateChanged(true);
 		}
@@ -463,12 +451,12 @@ bool UDataLayerEditorSubsystem::SetParentDataLayer(UDataLayerInstance* DataLayer
 	return false;
 }
 
-bool UDataLayerEditorSubsystem::AddActorToDataLayer(AActor* Actor, UDataLayerInstance* DataLayer)
+bool UDataLayerEditorSubsystem::AddActorToDataLayer(AActor* Actor, UDataLayerInstance* DataLayerInstance)
 {
 	TArray<AActor*> Actors;
 	Actors.Add(Actor);
 
-	return AddActorsToDataLayers(Actors, {DataLayer});
+	return AddActorsToDataLayers(Actors, { DataLayerInstance });
 }
 
 bool UDataLayerEditorSubsystem::AddActorToDataLayers(AActor* Actor, const TArray<UDataLayerInstance*>& DataLayers)
@@ -479,9 +467,9 @@ bool UDataLayerEditorSubsystem::AddActorToDataLayers(AActor* Actor, const TArray
 	return AddActorsToDataLayers(Actors, DataLayers);
 }
 
-bool UDataLayerEditorSubsystem::AddActorsToDataLayer(const TArray<AActor*>& Actors, UDataLayerInstance* DataLayer)
+bool UDataLayerEditorSubsystem::AddActorsToDataLayer(const TArray<AActor*>& Actors, UDataLayerInstance* DataLayerInstance)
 {
-	return AddActorsToDataLayers(Actors, {DataLayer});
+	return AddActorsToDataLayers(Actors, { DataLayerInstance });
 }
 
 bool UDataLayerEditorSubsystem::AddActorsToDataLayers(const TArray<AActor*>& Actors, const TArray<UDataLayerInstance*>& DataLayers)
@@ -492,9 +480,6 @@ bool UDataLayerEditorSubsystem::AddActorsToDataLayers(const TArray<AActor*>& Act
 	{
 		GEditor->GetSelectedActors()->BeginBatchSelectOperation();
 
-		UDataLayerSubsystem* DataLayerSubsystem = UWorld::GetSubsystem<UDataLayerSubsystem>(GetWorld());
-		UE_CLOG(!DataLayerSubsystem, LogDataLayerEditorSubsystem, Error, TEXT("%s - Adding actors to data layers while the world is null."), ANSI_TO_TCHAR(__FUNCTION__));
-
 		for (AActor* Actor : Actors)
 		{
 			if (!IsActorValidForDataLayer(Actor))
@@ -511,9 +496,10 @@ bool UDataLayerEditorSubsystem::AddActorsToDataLayers(const TArray<AActor*>& Act
 					// Make sure that a DataLayer Instance for this Data Layer Asset exists in the Actor's level.
 					if (Actor->GetLevel() != DataLayerInstance->GetTypedOuter<ULevel>())
 					{
-						if (DataLayerSubsystem != nullptr)
+						UDataLayerManager* DataLayerManager = UDataLayerManager::GetDataLayerManager(Actor);
+						if (ensureMsgf(DataLayerManager, TEXT("No DataLayerManager found for Actor %s, can't add actors to data layers."), *Actor->GetName()))
 						{
-							DataLayerInstance = DataLayerSubsystem->GetDataLayerInstance(DataLayerInstanceWithAsset->GetAsset(), Actor->GetLevel());
+							DataLayerInstance = DataLayerManager->GetDataLayerInstance(DataLayerInstanceWithAsset->GetAsset());
 
 							bool bDataLayerInstanceExistsInActorLevel = DataLayerInstance != nullptr;
 							if (!bDataLayerInstanceExistsInActorLevel)
@@ -524,10 +510,13 @@ bool UDataLayerEditorSubsystem::AddActorsToDataLayers(const TArray<AActor*>& Act
 					}
 				}
 
-				if (Actor->AddDataLayer(DataLayerInstance))
+				if (DataLayerInstance)
 				{
-					bActorWasModified = true;
-					BroadcastActorDataLayersChanged(Actor);
+					if (Actor->AddDataLayer(DataLayerInstance))
+					{
+						bActorWasModified = true;
+						BroadcastActorDataLayersChanged(Actor);
+					}
 				}
 			}
 
@@ -583,12 +572,12 @@ bool UDataLayerEditorSubsystem::RemoveActorsFromAllDataLayers(const TArray<AActo
 	return bChangesOccurred;
 }
 
-bool UDataLayerEditorSubsystem::RemoveActorFromDataLayer(AActor* Actor, UDataLayerInstance* DataLayer)
+bool UDataLayerEditorSubsystem::RemoveActorFromDataLayer(AActor* Actor, UDataLayerInstance* DataLayerInstance)
 {
 	TArray<AActor*> Actors;
 	Actors.Add(Actor);
 
-	return RemoveActorsFromDataLayers(Actors, {DataLayer});
+	return RemoveActorsFromDataLayers(Actors, { DataLayerInstance });
 }
 
 bool UDataLayerEditorSubsystem::RemoveActorFromDataLayers(AActor* Actor, const TArray<UDataLayerInstance*>& DataLayers)
@@ -599,9 +588,9 @@ bool UDataLayerEditorSubsystem::RemoveActorFromDataLayers(AActor* Actor, const T
 	return RemoveActorsFromDataLayers(Actors, DataLayers);
 }
 
-bool UDataLayerEditorSubsystem::RemoveActorsFromDataLayer(const TArray<AActor*>& Actors, UDataLayerInstance* DataLayer)
+bool UDataLayerEditorSubsystem::RemoveActorsFromDataLayer(const TArray<AActor*>& Actors, UDataLayerInstance* DataLayerInstance)
 {
-	return RemoveActorsFromDataLayers(Actors, {DataLayer});
+	return RemoveActorsFromDataLayers(Actors, { DataLayerInstance });
 }
 
 bool UDataLayerEditorSubsystem::RemoveActorsFromDataLayers(const TArray<AActor*>& Actors, const TArray<UDataLayerInstance*>& DataLayers)
@@ -617,12 +606,12 @@ bool UDataLayerEditorSubsystem::RemoveActorsFromDataLayers(const TArray<AActor*>
 		}
 
 		bool bActorWasModified = false;
-		for (const UDataLayerInstance* DataLayer : DataLayers)
+		for (const UDataLayerInstance* DataLayerInstance : DataLayers)
 		{
-			if (Actor->RemoveDataLayer(DataLayer))
+			if (Actor->RemoveDataLayer(DataLayerInstance))
 			{
 				bActorWasModified = true;
-				BroadcastDataLayerChanged(EDataLayerAction::Modify, DataLayer, NAME_None);
+				BroadcastDataLayerChanged(EDataLayerAction::Modify, DataLayerInstance, NAME_None);
 				BroadcastActorDataLayersChanged(Actor);
 			}
 		}
@@ -656,14 +645,14 @@ TArray<AActor*> UDataLayerEditorSubsystem::GetSelectedActors() const
 	return CurrentlySelectedActors;
 }
 
-bool UDataLayerEditorSubsystem::AddSelectedActorsToDataLayer(UDataLayerInstance* DataLayer)
+bool UDataLayerEditorSubsystem::AddSelectedActorsToDataLayer(UDataLayerInstance* DataLayerInstance)
 {
-	return AddActorsToDataLayer(GetSelectedActors(), DataLayer);
+	return AddActorsToDataLayer(GetSelectedActors(), DataLayerInstance);
 }
 
-bool UDataLayerEditorSubsystem::RemoveSelectedActorsFromDataLayer(UDataLayerInstance* DataLayer)
+bool UDataLayerEditorSubsystem::RemoveSelectedActorsFromDataLayer(UDataLayerInstance* DataLayerInstance)
 {
-	return RemoveActorsFromDataLayer(GetSelectedActors(), DataLayer);
+	return RemoveActorsFromDataLayer(GetSelectedActors(), DataLayerInstance);
 }
 
 bool UDataLayerEditorSubsystem::AddSelectedActorsToDataLayers(const TArray<UDataLayerInstance*>& DataLayers)
@@ -682,12 +671,12 @@ bool UDataLayerEditorSubsystem::RemoveSelectedActorsFromDataLayers(const TArray<
 // Operations on actors in DataLayers
 //
 ///////////////////////////////////////////////////////////////////////////////////////////////////////////////////////
-bool UDataLayerEditorSubsystem::SelectActorsInDataLayer(UDataLayerInstance* DataLayer, const bool bSelect, const bool bNotify, const bool bSelectEvenIfHidden)
+bool UDataLayerEditorSubsystem::SelectActorsInDataLayer(UDataLayerInstance* DataLayerInstance, const bool bSelect, const bool bNotify, const bool bSelectEvenIfHidden)
 {
-	return SelectActorsInDataLayer(DataLayer, bSelect, bNotify, bSelectEvenIfHidden, nullptr);
+	return SelectActorsInDataLayer(DataLayerInstance, bSelect, bNotify, bSelectEvenIfHidden, nullptr);
 }
 
-bool UDataLayerEditorSubsystem::SelectActorsInDataLayer(UDataLayerInstance* DataLayer, const bool bSelect, const bool bNotify, const bool bSelectEvenIfHidden, const TSharedPtr<FActorFilter>& Filter)
+bool UDataLayerEditorSubsystem::SelectActorsInDataLayer(UDataLayerInstance* DataLayerInstance, const bool bSelect, const bool bNotify, const bool bSelectEvenIfHidden, const TSharedPtr<FActorFilter>& Filter)
 {
 	bool bChangesOccurred = false;
 
@@ -706,9 +695,9 @@ bool UDataLayerEditorSubsystem::SelectActorsInDataLayer(UDataLayerInstance* Data
 				continue;
 			}
 
-			if (Actor->ContainsDataLayer(DataLayer))
+			if (Actor->ContainsDataLayer(DataLayerInstance))
 			{
-				// The actor was found to be in a specified DataLayer. Set selection state and move on to the next actor.
+				// The actor was found to be in a specified DataLayerInstance. Set selection state and move on to the next actor.
 				bool bNotifyForActor = false;
 				GEditor->GetSelectedActors()->Modify();
 				GEditor->SelectActor(Actor, bSelect, bNotifyForActor, bSelectEvenIfHidden);
@@ -754,11 +743,11 @@ bool UDataLayerEditorSubsystem::SelectActorsInDataLayers(const TArray<UDataLayer
 			continue;
 		}
 
-		for (const UDataLayerInstance* DataLayer : DataLayers)
+		for (const UDataLayerInstance* DataLayerInstance : DataLayers)
 		{
-			if (Actor->ContainsDataLayer(DataLayer))
+			if (Actor->ContainsDataLayer(DataLayerInstance))
 			{
-				// The actor was found to be in a specified DataLayer. Set selection state and move on to the next actor.
+				// The actor was found to be in a specified DataLayerInstance. Set selection state and move on to the next actor.
 				bool bNotifyForActor = false;
 				GEditor->GetSelectedActors()->Modify();
 				GEditor->SelectActor(Actor, bSelect, bNotifyForActor, bSelectEvenIfHidden);
@@ -798,7 +787,7 @@ bool UDataLayerEditorSubsystem::UpdateActorVisibility(AActor* Actor, bool& bOutS
 	TArray<const UDataLayerInstance*> DataLayerInstances = Actor->GetDataLayerInstances();
 	if (DataLayerInstances.IsEmpty())
 	{
-		// Actors that don't belong to any DataLayer shouldn't be hidden
+		// Actors that don't belong to any DataLayerInstance shouldn't be hidden
 		bOutActorModified = Actor->SetIsHiddenEdLayer(false);
 		return bOutActorModified;
 	}
@@ -818,7 +807,7 @@ bool UDataLayerEditorSubsystem::UpdateActorVisibility(AActor* Actor, bool& bOutS
 			}
 	}
 
-	// If the actor isn't part of a visible DataLayer, hide and de-select it.
+	// If the actor isn't part of a visible DataLayerInstance, hide and de-select it.
 	if (!bActorBelongsToVisibleDataLayer)
 	{
 		if (Actor->SetIsHiddenEdLayer(true))
@@ -890,12 +879,12 @@ bool UDataLayerEditorSubsystem::UpdateAllActorsVisibility(const bool bNotifySele
 //
 ///////////////////////////////////////////////////////////////////////////////////////////////////////////////////////
 
-void UDataLayerEditorSubsystem::AppendActorsFromDataLayer(UDataLayerInstance* DataLayer, TArray<AActor*>& InOutActors) const
+void UDataLayerEditorSubsystem::AppendActorsFromDataLayer(UDataLayerInstance* DataLayerInstance, TArray<AActor*>& InOutActors) const
 {
-	AppendActorsFromDataLayer(DataLayer, InOutActors, nullptr);
+	AppendActorsFromDataLayer(DataLayerInstance, InOutActors, nullptr);
 }
 
-void UDataLayerEditorSubsystem::AppendActorsFromDataLayer(UDataLayerInstance* DataLayer, TArray<AActor*>& InOutActors, const TSharedPtr<FActorFilter>& Filter) const
+void UDataLayerEditorSubsystem::AppendActorsFromDataLayer(UDataLayerInstance* DataLayerInstance, TArray<AActor*>& InOutActors, const TSharedPtr<FActorFilter>& Filter) const
 {
 	for (AActor* Actor : FActorRange(GetWorld()))
 	{
@@ -903,7 +892,7 @@ void UDataLayerEditorSubsystem::AppendActorsFromDataLayer(UDataLayerInstance* Da
 		{
 			continue;
 		}
-		if (Actor->ContainsDataLayer(DataLayer))
+		if (Actor->ContainsDataLayer(DataLayerInstance))
 		{
 			InOutActors.Add(Actor);
 		}
@@ -923,9 +912,9 @@ void UDataLayerEditorSubsystem::AppendActorsFromDataLayers(const TArray<UDataLay
 		{
 			continue;
 		}
-		for (const UDataLayerInstance* DataLayer : DataLayers)
+		for (const UDataLayerInstance* DataLayerInstance : DataLayers)
 		{
-			if (Actor->ContainsDataLayer(DataLayer))
+			if (Actor->ContainsDataLayer(DataLayerInstance))
 			{
 				InOutActors.Add(Actor);
 				break;
@@ -934,17 +923,17 @@ void UDataLayerEditorSubsystem::AppendActorsFromDataLayers(const TArray<UDataLay
 	}
 }
 
-TArray<AActor*> UDataLayerEditorSubsystem::GetActorsFromDataLayer(UDataLayerInstance* DataLayer) const
+TArray<AActor*> UDataLayerEditorSubsystem::GetActorsFromDataLayer(UDataLayerInstance* DataLayerInstance) const
 {
 	TArray<AActor*> OutActors;
-	AppendActorsFromDataLayer(DataLayer, OutActors);
+	AppendActorsFromDataLayer(DataLayerInstance, OutActors);
 	return OutActors;
 }
 
-TArray<AActor*> UDataLayerEditorSubsystem::GetActorsFromDataLayer(UDataLayerInstance* DataLayer, const TSharedPtr<FActorFilter>& Filter) const
+TArray<AActor*> UDataLayerEditorSubsystem::GetActorsFromDataLayer(UDataLayerInstance* DataLayerInstance, const TSharedPtr<FActorFilter>& Filter) const
 {
 	TArray<AActor*> OutActors;
-	AppendActorsFromDataLayer(DataLayer, OutActors, Filter);
+	AppendActorsFromDataLayer(DataLayerInstance, OutActors, Filter);
 	return OutActors;
 }
 
@@ -962,23 +951,23 @@ TArray<AActor*> UDataLayerEditorSubsystem::GetActorsFromDataLayers(const TArray<
 	return OutActors;
 }
 
-void UDataLayerEditorSubsystem::SetDataLayerVisibility(UDataLayerInstance* DataLayer, const bool bIsVisible)
+void UDataLayerEditorSubsystem::SetDataLayerVisibility(UDataLayerInstance* DataLayerInstance, const bool bIsVisible)
 {
-	SetDataLayersVisibility({ DataLayer }, bIsVisible);
+	SetDataLayersVisibility({ DataLayerInstance }, bIsVisible);
 }
 
 void UDataLayerEditorSubsystem::SetDataLayersVisibility(const TArray<UDataLayerInstance*>& DataLayers, const bool bIsVisible)
 {
 	bool bChangeOccurred = false;
-	for (UDataLayerInstance* DataLayer : DataLayers)
+	for (UDataLayerInstance* DataLayerInstance : DataLayers)
 	{
-		check(DataLayer);
+		check(DataLayerInstance);
 
-		if (DataLayer->IsVisible() != bIsVisible)
+		if (DataLayerInstance->IsVisible() != bIsVisible)
 		{
-			DataLayer->Modify(/*bAlswaysMarkDirty*/false);
-			DataLayer->SetVisible(bIsVisible);
-			BroadcastDataLayerChanged(EDataLayerAction::Modify, DataLayer, "bIsVisible");
+			DataLayerInstance->Modify(/*bAlswaysMarkDirty*/false);
+			DataLayerInstance->SetVisible(bIsVisible);
+			BroadcastDataLayerChanged(EDataLayerAction::Modify, DataLayerInstance, "bIsVisible");
 			bChangeOccurred = true;
 		}
 	}
@@ -989,10 +978,10 @@ void UDataLayerEditorSubsystem::SetDataLayersVisibility(const TArray<UDataLayerI
 	}
 }
 
-void UDataLayerEditorSubsystem::ToggleDataLayerVisibility(UDataLayerInstance* DataLayer)
+void UDataLayerEditorSubsystem::ToggleDataLayerVisibility(UDataLayerInstance* DataLayerInstance)
 {
-	check(DataLayer);
-	SetDataLayerVisibility(DataLayer, !DataLayer->IsVisible());
+	check(DataLayerInstance);
+	SetDataLayerVisibility(DataLayerInstance, !DataLayerInstance->IsVisible());
 }
 
 void UDataLayerEditorSubsystem::ToggleDataLayersVisibility(const TArray<UDataLayerInstance*>& DataLayers)
@@ -1002,11 +991,11 @@ void UDataLayerEditorSubsystem::ToggleDataLayersVisibility(const TArray<UDataLay
 		return;
 	}
 
-	for (UDataLayerInstance* DataLayer : DataLayers)
+	for (UDataLayerInstance* DataLayerInstance : DataLayers)
 	{
-		DataLayer->Modify();
-		DataLayer->SetVisible(!DataLayer->IsVisible());
-		BroadcastDataLayerChanged(EDataLayerAction::Modify, DataLayer, "bIsVisible");
+		DataLayerInstance->Modify();
+		DataLayerInstance->SetVisible(!DataLayerInstance->IsVisible());
+		BroadcastDataLayerChanged(EDataLayerAction::Modify, DataLayerInstance, "bIsVisible");
 	}
 
 	UpdateAllActorsVisibility(true, true);
@@ -1015,15 +1004,15 @@ void UDataLayerEditorSubsystem::ToggleDataLayersVisibility(const TArray<UDataLay
 void UDataLayerEditorSubsystem::MakeAllDataLayersVisible()
 {
 	UE_CLOG(!GetWorld(), LogDataLayerEditorSubsystem, Error, TEXT("%s - Failed because world in null."), ANSI_TO_TCHAR(__FUNCTION__));
-	if (UDataLayerSubsystem* DataLayerSubsystem = UWorld::GetSubsystem<UDataLayerSubsystem>(GetWorld()))
+	if (UDataLayerManager* DataLayerManager = UDataLayerManager::GetDataLayerManager(GetWorld()))
 	{
-		DataLayerSubsystem->ForEachDataLayer([this](UDataLayerInstance* DataLayer)
+		DataLayerManager->ForEachDataLayerInstance([this](UDataLayerInstance* DataLayerInstance)
 		{
-			if (!DataLayer->IsVisible())
+			if (!DataLayerInstance->IsVisible())
 			{
-				DataLayer->Modify();
-				DataLayer->SetVisible(true);
-				BroadcastDataLayerChanged(EDataLayerAction::Modify, DataLayer, "bIsVisible");
+				DataLayerInstance->Modify();
+				DataLayerInstance->SetVisible(true);
+				BroadcastDataLayerChanged(EDataLayerAction::Modify, DataLayerInstance, "bIsVisible");
 			}
 			return true;
 		});
@@ -1032,18 +1021,18 @@ void UDataLayerEditorSubsystem::MakeAllDataLayersVisible()
 }
 }
 
-bool UDataLayerEditorSubsystem::SetDataLayerIsLoadedInEditorInternal(UDataLayerInstance* DataLayer, const bool bIsLoadedInEditor, const bool bIsFromUserChange)
+bool UDataLayerEditorSubsystem::SetDataLayerIsLoadedInEditorInternal(UDataLayerInstance* DataLayerInstance, const bool bIsLoadedInEditor, const bool bIsFromUserChange)
 {
-	check(DataLayer);
-	if (DataLayer->IsLoadedInEditor() != bIsLoadedInEditor)
+	check(DataLayerInstance);
+	if (DataLayerInstance->IsLoadedInEditor() != bIsLoadedInEditor)
 	{
-		const bool bWasVisible = DataLayer->IsEffectiveVisible();
+		const bool bWasVisible = DataLayerInstance->IsEffectiveVisible();
 
-		DataLayer->Modify(false);
-		DataLayer->SetIsLoadedInEditor(bIsLoadedInEditor, /*bFromUserChange*/bIsFromUserChange);
-		BroadcastDataLayerChanged(EDataLayerAction::Modify, DataLayer, "bIsLoadedInEditor");
+		DataLayerInstance->Modify(false);
+		DataLayerInstance->SetIsLoadedInEditor(bIsLoadedInEditor, /*bFromUserChange*/bIsFromUserChange);
+		BroadcastDataLayerChanged(EDataLayerAction::Modify, DataLayerInstance, "bIsLoadedInEditor");
 
-		if (DataLayer->IsEffectiveVisible() != bWasVisible)
+		if (DataLayerInstance->IsEffectiveVisible() != bWasVisible)
 		{
 			UpdateAllActorsVisibility(true, true);
 		}
@@ -1052,9 +1041,9 @@ bool UDataLayerEditorSubsystem::SetDataLayerIsLoadedInEditorInternal(UDataLayerI
 	return false;
 }
 
-bool UDataLayerEditorSubsystem::SetDataLayerIsLoadedInEditor(UDataLayerInstance* DataLayer, const bool bIsLoadedInEditor, const bool bIsFromUserChange)
+bool UDataLayerEditorSubsystem::SetDataLayerIsLoadedInEditor(UDataLayerInstance* DataLayerInstance, const bool bIsLoadedInEditor, const bool bIsFromUserChange)
 {
-	if (SetDataLayerIsLoadedInEditorInternal(DataLayer, bIsLoadedInEditor, bIsFromUserChange))
+	if (SetDataLayerIsLoadedInEditorInternal(DataLayerInstance, bIsLoadedInEditor, bIsFromUserChange))
 	{
 		OnDataLayerEditorLoadingStateChanged(bIsFromUserChange);
 	}
@@ -1064,9 +1053,9 @@ bool UDataLayerEditorSubsystem::SetDataLayerIsLoadedInEditor(UDataLayerInstance*
 bool UDataLayerEditorSubsystem::SetDataLayersIsLoadedInEditor(const TArray<UDataLayerInstance*>& DataLayers, const bool bIsLoadedInEditor, const bool bIsFromUserChange)
 {
 	bool bChanged = false;
-	for (UDataLayerInstance* DataLayer : DataLayers)
+	for (UDataLayerInstance* DataLayerInstance : DataLayers)
 	{
-		bChanged |= SetDataLayerIsLoadedInEditorInternal(DataLayer, bIsLoadedInEditor, bIsFromUserChange);
+		bChanged |= SetDataLayerIsLoadedInEditorInternal(DataLayerInstance, bIsLoadedInEditor, bIsFromUserChange);
 	}
 	
 	if (bChanged)
@@ -1077,18 +1066,18 @@ bool UDataLayerEditorSubsystem::SetDataLayersIsLoadedInEditor(const TArray<UData
 	return true;
 }
 
-bool UDataLayerEditorSubsystem::ToggleDataLayerIsLoadedInEditor(UDataLayerInstance* DataLayer, const bool bIsFromUserChange)
+bool UDataLayerEditorSubsystem::ToggleDataLayerIsLoadedInEditor(UDataLayerInstance* DataLayerInstance, const bool bIsFromUserChange)
 {
-	check(DataLayer);
-	return SetDataLayerIsLoadedInEditor(DataLayer, !DataLayer->IsLoadedInEditor(), bIsFromUserChange);
+	check(DataLayerInstance);
+	return SetDataLayerIsLoadedInEditor(DataLayerInstance, !DataLayerInstance->IsLoadedInEditor(), bIsFromUserChange);
 }
 
 bool UDataLayerEditorSubsystem::ToggleDataLayersIsLoadedInEditor(const TArray<UDataLayerInstance*>& DataLayers, const bool bIsFromUserChange)
 {
 	bool bChanged = false;
-	for (UDataLayerInstance* DataLayer : DataLayers)
+	for (UDataLayerInstance* DataLayerInstance : DataLayers)
 	{
-		bChanged |= SetDataLayerIsLoadedInEditorInternal(DataLayer, !DataLayer->IsLoadedInEditor(), bIsFromUserChange);
+		bChanged |= SetDataLayerIsLoadedInEditorInternal(DataLayerInstance, !DataLayerInstance->IsLoadedInEditor(), bIsFromUserChange);
 	}
 	
 	if (bChanged)
@@ -1102,9 +1091,9 @@ bool UDataLayerEditorSubsystem::ToggleDataLayersIsLoadedInEditor(const TArray<UD
 TArray<UDataLayerInstance*> UDataLayerEditorSubsystem::GetAllDataLayers()
 {
 	TArray<UDataLayerInstance*> DataLayerInstances;
-	if (UDataLayerSubsystem* DataLayerSubsystem = UWorld::GetSubsystem<UDataLayerSubsystem>(GetWorld()))
+	if (UDataLayerManager* DataLayerManager = UDataLayerManager::GetDataLayerManager(GetWorld()))
 	{
-		DataLayerSubsystem->ForEachDataLayer([&DataLayerInstances](UDataLayerInstance* DataLayerInstance)
+		DataLayerManager->ForEachDataLayerInstance([&DataLayerInstances](UDataLayerInstance* DataLayerInstance)
 		{
 			DataLayerInstances.Add(DataLayerInstance);
 			return true;
@@ -1117,11 +1106,11 @@ bool UDataLayerEditorSubsystem::ResetUserSettings()
 {
 	bool bChanged = false;
 	UE_CLOG(!GetWorld(), LogDataLayerEditorSubsystem, Error, TEXT("%s - Failed because world in null."), ANSI_TO_TCHAR(__FUNCTION__));
-	if (UDataLayerSubsystem* DataLayerSubsystem = UWorld::GetSubsystem<UDataLayerSubsystem>(GetWorld()))
+	if (UDataLayerManager* DataLayerManager = UDataLayerManager::GetDataLayerManager(GetWorld()))
 	{
-		DataLayerSubsystem->ForEachDataLayer([this, &bChanged](UDataLayerInstance* DataLayer)
+		DataLayerManager->ForEachDataLayerInstance([this, &bChanged](UDataLayerInstance* DataLayerInstance)
 		{
-			bChanged |= SetDataLayerIsLoadedInEditorInternal(DataLayer, DataLayer->IsInitiallyLoadedInEditor(), true);
+			bChanged |= SetDataLayerIsLoadedInEditorInternal(DataLayerInstance, DataLayerInstance->IsInitiallyLoadedInEditor(), true);
 			return true;
 		});
 	
@@ -1145,18 +1134,18 @@ bool UDataLayerEditorSubsystem::HasDeprecatedDataLayers() const
 
 UDataLayerInstance* UDataLayerEditorSubsystem::GetDataLayerInstance(const FName& DataLayerInstanceName) const
 {
-	if (UDataLayerSubsystem* DataLayerSubsystem = UWorld::GetSubsystem<UDataLayerSubsystem>(GetWorld()))
+	if (UDataLayerManager* DataLayerManager = UDataLayerManager::GetDataLayerManager(GetWorld()))
 	{
-		return DataLayerSubsystem->GetDataLayerInstance(DataLayerInstanceName);
+		return const_cast<UDataLayerInstance*>(DataLayerManager->GetDataLayerInstance(DataLayerInstanceName));
 	}
 	return nullptr;
 }
 
 UDataLayerInstance* UDataLayerEditorSubsystem::GetDataLayerInstance(const UDataLayerAsset* DataLayerAsset) const
 {
-	if (UDataLayerSubsystem* DataLayerSubsystem = UWorld::GetSubsystem<UDataLayerSubsystem>(GetWorld()))
+	if (UDataLayerManager* DataLayerManager = UDataLayerManager::GetDataLayerManager(GetWorld()))
 	{
-		return DataLayerSubsystem->GetDataLayerInstance(DataLayerAsset);
+		return const_cast<UDataLayerInstance*>(DataLayerManager->GetDataLayerInstance(DataLayerAsset));
 	}
 	return nullptr;
 }
@@ -1176,13 +1165,13 @@ TArray<UDataLayerInstance*> UDataLayerEditorSubsystem::GetDataLayerInstances(con
 void UDataLayerEditorSubsystem::AddAllDataLayersTo(TArray<TWeakObjectPtr<UDataLayerInstance>>& OutDataLayers) const
 {
 	UE_CLOG(!GetWorld(), LogDataLayerEditorSubsystem, Error, TEXT("%s - Failed because world in null."), ANSI_TO_TCHAR(__FUNCTION__));
-	if (UDataLayerSubsystem* DataLayerSubsystem = UWorld::GetSubsystem<UDataLayerSubsystem>(GetWorld()))
+	if (UDataLayerManager* DataLayerManager = UDataLayerManager::GetDataLayerManager(GetWorld()))
 	{
-		DataLayerSubsystem->ForEachDataLayer([&OutDataLayers](UDataLayerInstance* DataLayer)
-			{
-				OutDataLayers.Add(DataLayer);
-				return true;
-			});
+		DataLayerManager->ForEachDataLayerInstance([&OutDataLayers](UDataLayerInstance* DataLayerInstance)
+		{
+			OutDataLayers.Add(DataLayerInstance);
+			return true;
+		});
 	}
 }
 
@@ -1214,24 +1203,30 @@ UDataLayerInstance* UDataLayerEditorSubsystem::CreateDataLayerInstance(const FDa
 void UDataLayerEditorSubsystem::DeleteDataLayers(const TArray<UDataLayerInstance*>& DataLayersToDelete)
 {
 	UE_CLOG(!GetWorld(), LogDataLayerEditorSubsystem, Error, TEXT("%s - Failed because world in null."), ANSI_TO_TCHAR(__FUNCTION__));
-	if (UDataLayerSubsystem* DataLayerSubsystem = UWorld::GetSubsystem<UDataLayerSubsystem>(GetWorld()))
+	
+	TArray<UDataLayerInstance*> DeletedDataLayerInstances;
+	for (UDataLayerInstance* DataLayerToDelete : DataLayersToDelete)
 	{
-		if (DataLayerSubsystem->RemoveDataLayers(DataLayersToDelete) > 0)
+		if (AWorldDataLayers* OuterWorldDataLayers = DataLayerToDelete ? DataLayerToDelete->GetOuterAWorldDataLayers() : nullptr)
 		{
-			for (UDataLayerInstance* DataLayerInstance : DataLayersToDelete)
+			if (OuterWorldDataLayers->RemoveDataLayer(DataLayerToDelete))
 			{
-				BroadcastDataLayerChanged(EDataLayerAction::Delete, DataLayerInstance, NAME_None);
+				DeletedDataLayerInstances.Add(DataLayerToDelete);
 			}
 		}
 	}
+	for (UDataLayerInstance* DeletedDataLayerInstance : DeletedDataLayerInstances)
+	{
+		BroadcastDataLayerChanged(EDataLayerAction::Delete, DeletedDataLayerInstance, NAME_None);
+	}
 }
 
 void UDataLayerEditorSubsystem::DeleteDataLayer(UDataLayerInstance* DataLayerToDelete)
 {
 	UE_CLOG(!GetWorld(), LogDataLayerEditorSubsystem, Error, TEXT("%s - Failed because world in null."), ANSI_TO_TCHAR(__FUNCTION__));
-	if (UDataLayerSubsystem* DataLayerSubsystem = UWorld::GetSubsystem<UDataLayerSubsystem>(GetWorld()))
+	if (AWorldDataLayers* OuterWorldDataLayers = DataLayerToDelete ? DataLayerToDelete->GetOuterAWorldDataLayers() : nullptr)
 	{
-		if (DataLayerSubsystem->RemoveDataLayer(DataLayerToDelete))
+		if (OuterWorldDataLayers->RemoveDataLayer(DataLayerToDelete))
 		{
 			BroadcastDataLayerChanged(EDataLayerAction::Delete, DataLayerToDelete, NAME_None);
 		}
@@ -1263,9 +1258,9 @@ void UDataLayerEditorSubsystem::BroadcastDataLayerEditorLoadingStateChanged(bool
 {
 	UE_CLOG(!GetWorld(), LogDataLayerEditorSubsystem, Error, TEXT("%s - Failed because world in null."), ANSI_TO_TCHAR(__FUNCTION__));
 	DataLayerEditorLoadingStateChanged.Broadcast(bIsFromUserChange);
-	if (UDataLayerSubsystem* DataLayerSubsystem = UWorld::GetSubsystem<UDataLayerSubsystem>(GetWorld()))
+	if (UDataLayerManager* DataLayerManager = UDataLayerManager::GetDataLayerManager(GetWorld()))
 	{
-		DataLayerSubsystem->UpdateDataLayerEditorPerProjectUserSettings();
+		DataLayerManager->UpdateDataLayerEditorPerProjectUserSettings();
 	}
 }
 
@@ -1320,9 +1315,10 @@ bool UDataLayerEditorSubsystem::TryGetDataLayerFromLabel(const FName& DataLayerL
 
 UDataLayerInstance* UDataLayerEditorSubsystem::GetDataLayerFromLabel(const FName& DataLayerLabel) const
 {
-	if (UDataLayerSubsystem* DataLayerSubsystem = UWorld::GetSubsystem<UDataLayerSubsystem>(GetWorld()))
+	if (AWorldDataLayers* WorldDataLayers = GetWorld()->GetWorldDataLayers())
 	{
-		return DataLayerSubsystem->GetDataLayerFromLabel(DataLayerLabel);
+		const UDataLayerInstance* DataLayerInstance = WorldDataLayers->GetDataLayerFromLabel(DataLayerLabel);
+		return const_cast<UDataLayerInstance*>(DataLayerInstance);
 	}
 	return nullptr;
 }
diff --git a/Engine/Source/Editor/DataLayerEditor/Private/DataLayer/DataLayerHierarchy.cpp b/Engine/Source/Editor/DataLayerEditor/Private/DataLayer/DataLayerHierarchy.cpp
--- a/Engine/Source/Editor/DataLayerEditor/Private/DataLayer/DataLayerHierarchy.cpp
+++ b/Engine/Source/Editor/DataLayerEditor/Private/DataLayer/DataLayerHierarchy.cpp
@@ -24,7 +24,7 @@
 #include "WorldDataLayersTreeItem.h"
 #include "WorldPartition/ActorDescContainerCollection.h"
 #include "WorldPartition/DataLayer/DataLayerInstance.h"
-#include "WorldPartition/DataLayer/DataLayerSubsystem.h"
+#include "WorldPartition/DataLayer/DataLayerManager.h"
 #include "WorldPartition/DataLayer/WorldDataLayers.h"
 #include "WorldPartition/IWorldPartitionEditorModule.h"
 #include "WorldPartition/WorldPartition.h"
@@ -161,137 +161,137 @@ UWorld* FDataLayerHierarchy::GetOwningWorld() const
 
 void FDataLayerHierarchy::CreateItems(TArray<FSceneOutlinerTreeItemPtr>& OutItems) const
 {
-	const UDataLayerSubsystem* DataLayerSubsystem = UWorld::GetSubsystem<UDataLayerSubsystem>(GetOwningWorld());
-
 	auto IsDataLayerShown = [this](const UDataLayerInstance* DataLayerInstance)
 	{
 		const bool bIsRuntimeDataLayer = DataLayerInstance->IsRuntime();
 		return ((bIsRuntimeDataLayer && bShowRuntimeDataLayers) || (!bIsRuntimeDataLayer && bShowEditorDataLayers)) && IsDataLayerPartOfSelection(DataLayerInstance);
 	};
 
-	auto CreateDataLayerTreeItems = [this, DataLayerSubsystem, &OutItems, IsDataLayerShown](const ULevel* OuterLevel)
+	auto CreateDataLayerTreeItems = [this, &OutItems, IsDataLayerShown](const UDataLayerManager* InDataLayerManager)
 	{
-		DataLayerSubsystem->ForEachDataLayer([this, &OutItems, IsDataLayerShown](UDataLayerInstance* DataLayer)
+		if (InDataLayerManager)
 		{
-			if (IsDataLayerShown(DataLayer))
+			InDataLayerManager->ForEachDataLayerInstance([this, &OutItems, IsDataLayerShown](UDataLayerInstance* DataLayerInstance)
 			{
-				if (FSceneOutlinerTreeItemPtr DataLayerItem = CreateDataLayerTreeItem(DataLayer))
+				if (IsDataLayerShown(DataLayerInstance))
 				{
-					OutItems.Add(DataLayerItem);
+					if (FSceneOutlinerTreeItemPtr DataLayerItem = CreateDataLayerTreeItem(DataLayerInstance))
+					{
+						OutItems.Add(DataLayerItem);
+					}
 				}
-			}
-			return true;
-		}, OuterLevel);
+				return true;
+			});
+		}
 	};
 
-	CreateDataLayerTreeItems(GetOwningWorld()->PersistentLevel);
-
-	// CurrentLevel represents the current level if different than the PersistentLevel
-	ULevel* CurrentLevel = (GetOwningWorld()->GetCurrentLevel() && !GetOwningWorld()->GetCurrentLevel()->IsPersistentLevel()) ? GetOwningWorld()->GetCurrentLevel() : nullptr;
-	if (CurrentLevel != nullptr)
+	if (GetOwningWorld()->IsPlayInEditor())
 	{
-		CreateDataLayerTreeItems(CurrentLevel);
+		// PIE loops on each DataLayerManager of each registered WorldPartition. 
+		// For performance reasons, child actors are not shown.
+		UWorldPartitionSubsystem* WorldPartitionSubsystem = UWorld::GetSubsystem<UWorldPartitionSubsystem>(GetOwningWorld());
+		WorldPartitionSubsystem->ForEachWorldPartition([&CreateDataLayerTreeItems](UWorldPartition* WorldPartition)
+		{
+			CreateDataLayerTreeItems(WorldPartition->GetDataLayerManager());
+			return true;
+		});
 	}
-
-	// Create Tree items for WorldDataLayers(WDL) which are part of the current world (not subworld) but are not the ULevel::AWorldDataLayers.
-	// The only way to create data layer instances for in these WDL is to right-click on the WDL tree item and select create new.
-	// We cannot rely on creating the WDL tree item only if it has DataLayerInstance (::FindOrCreateParentItem), as it prevents to initially create a new instance.
-	DataLayerSubsystem->ForEachWorldDataLayer([this, &OutItems](AWorldDataLayers* WorldDataLayers)
+	else
 	{
-		if (!WorldDataLayers->IsSubWorldDataLayers() && !WorldDataLayers->IsTheMainWorldDataLayers())
-		{
-			FSceneOutlinerTreeItemPtr WorldDataLayerItem = Mode->CreateItemFor<FWorldDataLayersTreeItem>(WorldDataLayers, true);
-			OutItems.Add(WorldDataLayerItem);
-		}
+		// CurrentLevel represents the current level if different than the PersistentLevel
+		ULevel* CurrentLevel = (GetOwningWorld()->GetCurrentLevel() && !GetOwningWorld()->GetCurrentLevel()->IsPersistentLevel()) ? GetOwningWorld()->GetCurrentLevel() : nullptr;
 
-		return true;
-	});
+		// Create DataLayerTreeItems for World and for Current Level (if any).
+		CreateDataLayerTreeItems(GetOwningWorld()->GetDataLayerManager());
+		CreateDataLayerTreeItems(UDataLayerManager::GetDataLayerManager(CurrentLevel));
 
-	if (bShowDataLayerActors)
-	{
-		const ULevelInstanceSubsystem* LevelInstanceSubsystem = UWorld::GetSubsystem<ULevelInstanceSubsystem>(GetOwningWorld());
-
-		// Build a set of considered level instances to show nested level instance actors and their child actors
-		TSet<const ILevelInstanceInterface*> ConsideredLevelInstances;		
-		if (bShowLevelInstanceContent && CurrentLevel && LevelInstanceSubsystem)
+		if (bShowDataLayerActors)
 		{
-			ILevelInstanceInterface* CurrentLevelInstance = LevelInstanceSubsystem->GetOwningLevelInstance(CurrentLevel);
-			if (LevelInstanceSubsystem->IsEditingLevelInstance(CurrentLevelInstance))
+			const ULevelInstanceSubsystem* LevelInstanceSubsystem = UWorld::GetSubsystem<ULevelInstanceSubsystem>(GetOwningWorld());
+
+			// Build a set of considered level instances to show nested level instance actors and their child actors
+			TSet<const ILevelInstanceInterface*> ConsideredLevelInstances;		
+			if (bShowLevelInstanceContent && CurrentLevel && LevelInstanceSubsystem)
 			{
-				const bool bRecursive = true;
-				LevelInstanceSubsystem->ForEachLevelInstanceChild(CurrentLevelInstance, bRecursive, [&ConsideredLevelInstances](const ILevelInstanceInterface* ChildLevelInstance)
+				ILevelInstanceInterface* CurrentLevelInstance = LevelInstanceSubsystem->GetOwningLevelInstance(CurrentLevel);
+				if (LevelInstanceSubsystem->IsEditingLevelInstance(CurrentLevelInstance))
 				{
-					ConsideredLevelInstances.Add(ChildLevelInstance);
-					return true;
-				});
+					const bool bRecursive = true;
+					LevelInstanceSubsystem->ForEachLevelInstanceChild(CurrentLevelInstance, bRecursive, [&ConsideredLevelInstances](const ILevelInstanceInterface* ChildLevelInstance)
+					{
+						ConsideredLevelInstances.Add(ChildLevelInstance);
+						return true;
+					});
+				}
 			}
-		}
 
-		for (AActor* Actor : FActorRange(GetOwningWorld()))
-		{
-			// Consider all actors or actors part of current level (if there is one)
-			bool bConsiderActor = !CurrentLevel || (Actor->GetLevel() == CurrentLevel);
-			if (!bConsiderActor && bShowLevelInstanceContent && LevelInstanceSubsystem)
+			for (AActor* Actor : FActorRange(GetOwningWorld()))
 			{
-				// If current level is referenced by editing level instance, also consider child actors part of any child Level Instances of this level
-				ILevelInstanceInterface* ParentLevelInstance = LevelInstanceSubsystem->GetParentLevelInstance(Actor);
-				bConsiderActor = ParentLevelInstance && ConsideredLevelInstances.Contains(ParentLevelInstance);
-			}
+				// Consider all actors or actors part of current level (if there is one)
+				bool bConsiderActor = !CurrentLevel || (Actor->GetLevel() == CurrentLevel);
+				if (!bConsiderActor && bShowLevelInstanceContent && LevelInstanceSubsystem)
+				{
+					// If current level is referenced by editing level instance, also consider child actors part of any child Level Instances of this level
+					ILevelInstanceInterface* ParentLevelInstance = LevelInstanceSubsystem->GetParentLevelInstance(Actor);
+					bConsiderActor = ParentLevelInstance && ConsideredLevelInstances.Contains(ParentLevelInstance);
+				}
 
-			if (bConsiderActor)
-			{
-				for (const UDataLayerInstance* DataLayerInstance : Actor->GetDataLayerInstances())
+				if (bConsiderActor)
 				{
-					if (IsDataLayerShown(DataLayerInstance))
+					for (const UDataLayerInstance* DataLayerInstance : Actor->GetDataLayerInstances())
 					{
-						if (FSceneOutlinerTreeItemPtr DataLayerActorItem = Mode->CreateItemFor<FDataLayerActorTreeItem>(FDataLayerActorTreeItemData(Actor, const_cast<UDataLayerInstance*>(DataLayerInstance))))
+						if (IsDataLayerShown(DataLayerInstance))
 						{
-							OutItems.Add(DataLayerActorItem);
+							if (FSceneOutlinerTreeItemPtr DataLayerActorItem = Mode->CreateItemFor<FDataLayerActorTreeItem>(FDataLayerActorTreeItemData(Actor, const_cast<UDataLayerInstance*>(DataLayerInstance))))
+							{
+								OutItems.Add(DataLayerActorItem);
+							}
 						}
 					}
 				}
 			}
-		}
 
-		if (bShowUnloadedActors)
-		{
-			if (UWorldPartitionSubsystem* WorldPartitionSubsystem = UWorld::GetSubsystem<UWorldPartitionSubsystem>(GetOwningWorld()))
+			if (bShowUnloadedActors)
 			{
-				WorldPartitionSubsystem->ForEachWorldPartition([this, CurrentLevel, LevelInstanceSubsystem, DataLayerSubsystem, IsDataLayerShown, &OutItems](UWorldPartition* WorldPartition)
+				if (UWorldPartitionSubsystem* WorldPartitionSubsystem = UWorld::GetSubsystem<UWorldPartitionSubsystem>(GetOwningWorld()))
 				{
-					// Skip WorldPartition if it's not the one of the current level (the editing level instance)
-					// or if we hide the content of level instances
-					UWorld* OuterWorld = WorldPartition->GetTypedOuter<UWorld>();
-					ULevel* OuterLevel = OuterWorld ? OuterWorld->PersistentLevel : nullptr;
-					if ((CurrentLevel && (CurrentLevel != OuterLevel)) ||
-						(!bShowLevelInstanceContent && OuterLevel && LevelInstanceSubsystem && LevelInstanceSubsystem->GetOwningLevelInstance(OuterLevel)))
+					WorldPartitionSubsystem->ForEachWorldPartition([this, CurrentLevel, LevelInstanceSubsystem, IsDataLayerShown, &OutItems](UWorldPartition* WorldPartition)
 					{
-						return true;
-					}
+						// Skip WorldPartition if it's not the one of the current level (the editing level instance)
+						// or if we hide the content of level instances
+						UWorld* OuterWorld = WorldPartition->GetTypedOuter<UWorld>();
+						ULevel* OuterLevel = OuterWorld ? OuterWorld->PersistentLevel : nullptr;
+						if ((CurrentLevel && (CurrentLevel != OuterLevel)) ||
+							(!bShowLevelInstanceContent && OuterLevel && LevelInstanceSubsystem && LevelInstanceSubsystem->GetOwningLevelInstance(OuterLevel)))
+						{
+							return true;
+						}
 
-					// Create an FDataLayerActorDescTreeItem for each unloaded actor of this WorldPartition
-					FWorldPartitionHelpers::ForEachActorDesc(WorldPartition, [this, DataLayerSubsystem, IsDataLayerShown, &OutItems](const FWorldPartitionActorDesc* ActorDesc)
-					{
-						if (ActorDesc != nullptr && !ActorDesc->IsLoaded())
+						// Create an FDataLayerActorDescTreeItem for each unloaded actor of this WorldPartition
+						const UDataLayerManager* DataLayerManager = WorldPartition->GetDataLayerManager();
+						FWorldPartitionHelpers::ForEachActorDesc(WorldPartition, [this, IsDataLayerShown, DataLayerManager, &OutItems](const FWorldPartitionActorDesc* ActorDesc)
 						{
-							for (const FName& DataLayerInstanceName : ActorDesc->GetDataLayerInstanceNames())
+							if (ActorDesc != nullptr && !ActorDesc->IsLoaded())
 							{
-								if (const UDataLayerInstance* const DataLayerInstance = DataLayerSubsystem->GetDataLayerInstance(DataLayerInstanceName))
+								for (const FName& DataLayerInstanceName : ActorDesc->GetDataLayerInstanceNames())
 								{
-									if (IsDataLayerShown(DataLayerInstance))
+									if (const UDataLayerInstance* const DataLayerInstance = DataLayerManager->GetDataLayerInstance(DataLayerInstanceName))
 									{
-										if (const FSceneOutlinerTreeItemPtr ActorDescItem = Mode->CreateItemFor<FDataLayerActorDescTreeItem>(FDataLayerActorDescTreeItemData(ActorDesc->GetGuid(), ActorDesc->GetContainer(), const_cast<UDataLayerInstance*>(DataLayerInstance))))
+										if (IsDataLayerShown(DataLayerInstance))
 										{
-											OutItems.Add(ActorDescItem);
+											if (const FSceneOutlinerTreeItemPtr ActorDescItem = Mode->CreateItemFor<FDataLayerActorDescTreeItem>(FDataLayerActorDescTreeItemData(ActorDesc->GetGuid(), ActorDesc->GetContainer(), const_cast<UDataLayerInstance*>(DataLayerInstance))))
+											{
+												OutItems.Add(ActorDescItem);
+											}
 										}
 									}
 								}
 							}
-						}
+							return true;
+						});
 						return true;
 					});
-					return true;
-				});
+				}
 			}
 		}
 	}
@@ -557,10 +557,10 @@ void FDataLayerHierarchy::OnActorDescAdded(FWorldPartitionActorDesc* InActorDesc
 		return;
 	}
 
-	const UDataLayerSubsystem* const DataLayerSubsystem = UWorld::GetSubsystem<UDataLayerSubsystem>(GetOwningWorld());
+	const UDataLayerManager* const DataLayerManager = UDataLayerManager::GetDataLayerManager(GetOwningWorld());
 	const TArray<FName>& DataLayerInstanceNames = InActorDesc->GetDataLayerInstanceNames();
 
-	if (DataLayerSubsystem && DataLayerInstanceNames.Num() > 0)
+	if (DataLayerManager && DataLayerInstanceNames.Num() > 0)
 	{
 		FSceneOutlinerHierarchyChangedData EventData;
 		EventData.Type = FSceneOutlinerHierarchyChangedData::Added;
@@ -568,7 +568,7 @@ void FDataLayerHierarchy::OnActorDescAdded(FWorldPartitionActorDesc* InActorDesc
 		EventData.Items.Reserve(DataLayerInstanceNames.Num());
 		for (const FName& DataLayerInstanceName : DataLayerInstanceNames)
 		{
-			const UDataLayerInstance* const DataLayerInstance = DataLayerSubsystem->GetDataLayerInstance(DataLayerInstanceName);
+			const UDataLayerInstance* const DataLayerInstance = DataLayerManager->GetDataLayerInstance(DataLayerInstanceName);
 			EventData.Items.Add(Mode->CreateItemFor<FDataLayerActorDescTreeItem>(FDataLayerActorDescTreeItemData(InActorDesc->GetGuid(), InActorDesc->GetContainer(), const_cast<UDataLayerInstance*>(DataLayerInstance))));
 		}
 		HierarchyChangedEvent.Broadcast(EventData);
@@ -582,10 +582,10 @@ void FDataLayerHierarchy::OnActorDescRemoved(FWorldPartitionActorDesc* InActorDe
 		return;
 	}
 
-	const UDataLayerSubsystem* const DataLayerSubsystem = UWorld::GetSubsystem<UDataLayerSubsystem>(GetOwningWorld());
+	const UDataLayerManager* const DataLayerManager = UDataLayerManager::GetDataLayerManager(GetOwningWorld());
 	const TArray<FName>& DataLayerInstanceNames = InActorDesc->GetDataLayerInstanceNames();
 	
-	if (DataLayerSubsystem && DataLayerInstanceNames.Num() > 0)
+	if (DataLayerManager && DataLayerInstanceNames.Num() > 0)
 	{
 		FSceneOutlinerHierarchyChangedData EventData;
 		EventData.Type = FSceneOutlinerHierarchyChangedData::Removed;
@@ -593,7 +593,7 @@ void FDataLayerHierarchy::OnActorDescRemoved(FWorldPartitionActorDesc* InActorDe
 
 		for (const FName& DataLayerInstanceName : DataLayerInstanceNames)
 		{
-			const UDataLayerInstance* const DataLayer = DataLayerSubsystem->GetDataLayerInstance(DataLayerInstanceName);
+			const UDataLayerInstance* const DataLayer = DataLayerManager->GetDataLayerInstance(DataLayerInstanceName);
 			EventData.ItemIDs.Add(FDataLayerActorDescTreeItem::ComputeTreeItemID(InActorDesc->GetGuid(), InActorDesc->GetContainer(), DataLayer));
 		}
 		HierarchyChangedEvent.Broadcast(EventData);
diff --git a/Engine/Source/Editor/DataLayerEditor/Private/DataLayer/DataLayerMode.cpp b/Engine/Source/Editor/DataLayerEditor/Private/DataLayer/DataLayerMode.cpp
--- a/Engine/Source/Editor/DataLayerEditor/Private/DataLayer/DataLayerMode.cpp
+++ b/Engine/Source/Editor/DataLayerEditor/Private/DataLayer/DataLayerMode.cpp
@@ -88,7 +88,7 @@
 #include "WorldPartition/DataLayer/DataLayerAsset.h"
 #include "WorldPartition/DataLayer/DataLayerInstance.h"
 #include "WorldPartition/DataLayer/DataLayerInstanceWithAsset.h"
-#include "WorldPartition/DataLayer/DataLayerSubsystem.h"
+#include "WorldPartition/DataLayer/DataLayerManager.h"
 #include "WorldPartition/DataLayer/WorldDataLayers.h"
 #include "WorldPartition/WorldPartitionEditorPerProjectUserSettings.h"
 #include "WorldTreeItem.h"
@@ -1244,15 +1244,15 @@ void FDataLayerMode::RegisterContextMenu()
 
 			bool bHasActorEditorContextDataLayers = false;			
 			TArray<const UDataLayerInstance*> AllDataLayers;
-			if (const UDataLayerSubsystem* DataLayerSubsystem = UWorld::GetSubsystem<UDataLayerSubsystem>(Mode->GetOwningWorld()))
+			if (const UDataLayerManager* DataLayerManager = UDataLayerManager::GetDataLayerManager(Mode->GetOwningWorld()))
 			{
-				DataLayerSubsystem->ForEachDataLayer([&AllDataLayers](UDataLayerInstance* DataLayerInstance)
+				DataLayerManager->ForEachDataLayerInstance([&AllDataLayers](UDataLayerInstance* DataLayerInstance)
 				{
 					AllDataLayers.Add(DataLayerInstance);
 					return true;
 				});
 
-				bHasActorEditorContextDataLayers = !DataLayerSubsystem->GetActorEditorContextDataLayers().IsEmpty();
+				bHasActorEditorContextDataLayers = !DataLayerManager->GetActorEditorContextDataLayers().IsEmpty();
 			}
 
 			{
@@ -1681,9 +1681,9 @@ void FDataLayerMode::CreateViewContent(FMenuBuilder& MenuBuilder)
 	);
 
 	TArray<UDataLayerInstance*> AllDataLayers;
-	if (const UDataLayerSubsystem* DataLayerSubsystem = UWorld::GetSubsystem<UDataLayerSubsystem>(GetOwningWorld()))
+	if (const UDataLayerManager* DataLayerManager = UDataLayerManager::GetDataLayerManager(GetOwningWorld()))
 	{
-		DataLayerSubsystem->ForEachDataLayer([&AllDataLayers](UDataLayerInstance* DataLayer)
+		DataLayerManager->ForEachDataLayerInstance([&AllDataLayers](UDataLayerInstance* DataLayer)
 		{
 			AllDataLayers.Add(DataLayer);
 			return true;
diff --git a/Engine/Source/Editor/DataLayerEditor/Private/DataLayer/DataLayersActorDescTreeItem.h b/Engine/Source/Editor/DataLayerEditor/Private/DataLayer/DataLayersActorDescTreeItem.h
--- a/Engine/Source/Editor/DataLayerEditor/Private/DataLayer/DataLayersActorDescTreeItem.h
+++ b/Engine/Source/Editor/DataLayerEditor/Private/DataLayer/DataLayersActorDescTreeItem.h
@@ -8,6 +8,7 @@
 #include "ActorDescTreeItem.h"
 #include "WorldPartition/DataLayer/DataLayerInstance.h"
 #include "WorldPartition/ActorDescContainerCollection.h"
+#include "WorldPartition/WorldPartition.h"
 #include "LevelInstance/LevelInstanceSubsystem.h"
 #include "Misc/ArchiveMD5.h"
 #include "Misc/SecureHash.h"
@@ -49,7 +50,7 @@ public:
 		if (ActorDescHandle.IsValid())
 		{
 			UActorDescContainer* Container = ActorDescHandle->GetContainer();
-			UWorld* OwningWorld = Container->GetWorld();
+			UWorld* OwningWorld = Container->GetWorldPartition()->GetWorld();
 			ULevelInstanceSubsystem* LevelInstanceSubsystem = UWorld::GetSubsystem<ULevelInstanceSubsystem>(OwningWorld);
 			ULevel* Level = Container->GetTypedOuter<UWorld>()->PersistentLevel;
 			if (LevelInstanceSubsystem && Level)
@@ -82,7 +83,8 @@ public:
 	{
 		if (InContainer)
 		{
-			if (ULevelInstanceSubsystem* LevelInstanceSubsystem = UWorld::GetSubsystem<ULevelInstanceSubsystem>(InContainer->GetWorld()))
+			UWorld* OwningWorld = InContainer->GetWorldPartition()->GetWorld();
+			if (ULevelInstanceSubsystem* LevelInstanceSubsystem = UWorld::GetSubsystem<ULevelInstanceSubsystem>(OwningWorld))
 			{
 				ULevel* Level = InContainer->GetTypedOuter<UWorld>()->PersistentLevel;
 				return LevelInstanceSubsystem->GetParentLevelInstanceActors(Level);
diff --git a/Engine/Source/Editor/DataLayerEditor/Private/DataLayer/SDataLayerTreeLabel.cpp b/Engine/Source/Editor/DataLayerEditor/Private/DataLayer/SDataLayerTreeLabel.cpp
--- a/Engine/Source/Editor/DataLayerEditor/Private/DataLayer/SDataLayerTreeLabel.cpp
+++ b/Engine/Source/Editor/DataLayerEditor/Private/DataLayer/SDataLayerTreeLabel.cpp
@@ -35,7 +35,6 @@
 #include "Widgets/Text/STextBlock.h"
 #include "Widgets/Views/STableRow.h"
 #include "WorldPartition/DataLayer/DataLayerInstance.h"
-#include "WorldPartition/DataLayer/DataLayerSubsystem.h"
 #include "WorldPartition/DataLayer/DataLayerInstanceWithAsset.h"
 
 struct FSlateBrush;
@@ -169,15 +168,14 @@ FSlateFontInfo SDataLayerTreeLabel::GetDisplayNameFont() const
 
 FText SDataLayerTreeLabel::GetDisplayText() const
 {
-	const UDataLayerInstance* DataLayer = DataLayerPtr.Get();
+	const UDataLayerInstance* DataLayerInstance = DataLayerPtr.Get();
 	bool bIsDataLayerActive = false;
 	FText SuffixText = FText::GetEmpty();
 	if (!bInEditingMode)
 	{
-		if (DataLayer && DataLayer->IsRuntime() && DataLayer->GetWorld()->IsPlayInEditor())
+		if (DataLayerInstance && DataLayerInstance->IsRuntime() && DataLayerInstance->GetWorld()->IsPlayInEditor())
 		{
-			const UDataLayerSubsystem* DataLayerSubsystem = DataLayer->GetWorld()->GetSubsystem<UDataLayerSubsystem>();
-			SuffixText = FText::Format(LOCTEXT("DataLayerRuntimeState", " ({0})"), FTextStringHelper::CreateFromBuffer(GetDataLayerRuntimeStateName(DataLayerSubsystem->GetDataLayerEffectiveRuntimeState(DataLayer))));
+			SuffixText = FText::Format(LOCTEXT("DataLayerRuntimeState", " ({0})"), FTextStringHelper::CreateFromBuffer(GetDataLayerRuntimeStateName(DataLayerInstance->GetEffectiveRuntimeState())));
 		}
 		else if (IsInActorEditorContext())
 		{
@@ -185,9 +183,9 @@ FText SDataLayerTreeLabel::GetDisplayText() const
 		}
 	}
 
-	if (DataLayer)
+	if (DataLayerInstance)
 	{
-		if (const UDataLayerInstanceWithAsset* DataLayerInstanceWithAsset = Cast<UDataLayerInstanceWithAsset>(DataLayer))
+		if (const UDataLayerInstanceWithAsset* DataLayerInstanceWithAsset = Cast<UDataLayerInstanceWithAsset>(DataLayerInstance))
 		{
 			if (!DataLayerInstanceWithAsset->GetAsset())
 			{
@@ -195,7 +193,7 @@ FText SDataLayerTreeLabel::GetDisplayText() const
 			}
 		}
 
-		return FText::Format(LOCTEXT("DataLayerDisplayText", "{0}{1}"), FText::FromString(DataLayer->GetDataLayerShortName()), SuffixText);
+		return FText::Format(LOCTEXT("DataLayerDisplayText", "{0}{1}"), FText::FromString(DataLayerInstance->GetDataLayerShortName()), SuffixText);
 	}
 	
 	static const FText DataLayerDeleted = LOCTEXT("DataLayerLabelForMissingDataLayer", "(Deleted Data Layer)");
@@ -255,15 +253,14 @@ FSlateColor SDataLayerTreeLabel::GetForegroundColor() const
 		return BaseColor.GetValue();
 	}
 
-	const UDataLayerInstance* DataLayer = DataLayerPtr.Get();
-	if (DataLayer)
+	const UDataLayerInstance* DataLayerInstance = DataLayerPtr.Get();
+	if (DataLayerInstance)
 	{
-		if (DataLayer->GetWorld()->IsPlayInEditor())
+		if (DataLayerInstance->GetWorld()->IsPlayInEditor())
 		{
-			if (DataLayer->IsRuntime())
+			if (DataLayerInstance->IsRuntime())
 			{
-				const UDataLayerSubsystem* DataLayerSubsystem = DataLayer->GetWorld()->GetSubsystem<UDataLayerSubsystem>();
-				EDataLayerRuntimeState State = DataLayerSubsystem->GetDataLayerEffectiveRuntimeState(DataLayer);
+				EDataLayerRuntimeState State = DataLayerInstance->GetEffectiveRuntimeState();
 				switch (State)
 				{
 				case EDataLayerRuntimeState::Activated:
@@ -279,12 +276,12 @@ FSlateColor SDataLayerTreeLabel::GetForegroundColor() const
 				return FSceneOutlinerCommonLabelData::DarkColor;
 			}
 		}
-		else if (DataLayer->IsLocked())
+		else if (DataLayerInstance->IsLocked())
 		{
 			return FSceneOutlinerCommonLabelData::DarkColor;
 		}
 	}
-	if (!DataLayer || !DataLayer->GetWorld())
+	if (!DataLayerInstance || !DataLayerInstance->GetWorld())
 	{
 		return FLinearColor(0.2f, 0.2f, 0.25f);
 	}
diff --git a/Engine/Source/Editor/DataLayerEditor/Public/DataLayer/DataLayerEditorSubsystem.h b/Engine/Source/Editor/DataLayerEditor/Public/DataLayer/DataLayerEditorSubsystem.h
--- a/Engine/Source/Editor/DataLayerEditor/Public/DataLayer/DataLayerEditorSubsystem.h
+++ b/Engine/Source/Editor/DataLayerEditor/Public/DataLayer/DataLayerEditorSubsystem.h
@@ -110,9 +110,6 @@ public:
 	void AddToActorEditorContext(UDataLayerInstance* InDataLayerInstance);
 	void RemoveFromActorEditorContext(UDataLayerInstance* InDataLayerInstance);
 
-	void OnWorldDataLayersPostRegister(AWorldDataLayers* WorldDataLayers);
-	void OnWorldDataLayersPreUnregister(AWorldDataLayers* WorldDataLayers);
-
 	TArray<const UDataLayerInstance*> GetDataLayerInstances(const TArray<const UDataLayerAsset*> DataLayerAssets) const;
 
 	template<class DataLayerType, typename ...Args>
diff --git a/Engine/Source/Editor/SceneOutliner/Private/ActorBrowsingMode.cpp b/Engine/Source/Editor/SceneOutliner/Private/ActorBrowsingMode.cpp
--- a/Engine/Source/Editor/SceneOutliner/Private/ActorBrowsingMode.cpp
+++ b/Engine/Source/Editor/SceneOutliner/Private/ActorBrowsingMode.cpp
@@ -39,7 +39,7 @@
 #include "LevelInstance/LevelInstanceEditorInstanceActor.h"
 #include "WorldPartition/WorldPartition.h"
 #include "WorldPartition/WorldPartitionSubsystem.h"
-#include "WorldPartition/DataLayer/DataLayerSubsystem.h"
+#include "WorldPartition/DataLayer/DataLayerManager.h"
 #include "WorldPartition/IWorldPartitionEditorModule.h"
 #include "Subsystems/ActorEditorContextSubsystem.h"
 #include "ISourceControlModule.h"
@@ -154,20 +154,17 @@ FActorBrowsingMode::FActorBrowsingMode(SSceneOutliner* InSceneOutliner, TWeakObj
 			FActorBrowsingModeConfig* Settings = GetMutableConfig();
 			if (Settings && Settings->bShowOnlyActorsInCurrentDataLayers)
 			{
-				const UDataLayerSubsystem* DataLayerSubsystem = RepresentingWorld.IsValid() ? RepresentingWorld->GetSubsystem<UDataLayerSubsystem>() : nullptr;
-				if (!DataLayerSubsystem || DataLayerSubsystem->GetActorEditorContextDataLayers().IsEmpty())
+				const UDataLayerManager* DataLayerManager = UDataLayerManager::GetDataLayerManager(RepresentingWorld.Get());
+				if (!DataLayerManager || DataLayerManager->GetActorEditorContextDataLayers().IsEmpty())
 				{
 					return true;
 				}
 				
-				for (const FName& DataLayerInstanceName : ActorDesc->GetDataLayerInstanceNames())
+				for(const UDataLayerInstance* const DataLayerInstance : DataLayerManager->GetDataLayerInstances(ActorDesc->GetDataLayerInstanceNames()))
 				{
-					if (const UDataLayerInstance* const DataLayerInstance = DataLayerSubsystem->GetDataLayerInstance(DataLayerInstanceName))
+					if (DataLayerInstance->IsInActorEditorContext())
 					{
-						if (DataLayerInstance->IsInActorEditorContext())
-						{
-							return true;
-						}
+						return true;
 					}
 				}
 				return false;
@@ -507,8 +504,8 @@ TSharedRef<FSceneOutlinerFilter> FActorBrowsingMode::CreateIsInCurrentDataLayers
 {
 	return MakeShareable(new FActorFilter(FActorTreeItem::FFilterPredicate::CreateStatic([](const AActor* InActor)
 		{
-			const UDataLayerSubsystem* DataLayerSubsystem = InActor->GetWorld() ? InActor->GetWorld()->GetSubsystem<UDataLayerSubsystem>() : nullptr;
-			if (!DataLayerSubsystem || DataLayerSubsystem->GetActorEditorContextDataLayers().IsEmpty())
+			const UDataLayerManager* DataLayerManager = UDataLayerManager::GetDataLayerManager(InActor->GetWorld());
+			if (!DataLayerManager || DataLayerManager->GetActorEditorContextDataLayers().IsEmpty())
 			{
 				return true;
 			}
diff --git a/Engine/Source/Editor/SceneOutliner/Private/ActorDescTreeItem.cpp b/Engine/Source/Editor/SceneOutliner/Private/ActorDescTreeItem.cpp
--- a/Engine/Source/Editor/SceneOutliner/Private/ActorDescTreeItem.cpp
+++ b/Engine/Source/Editor/SceneOutliner/Private/ActorDescTreeItem.cpp
@@ -320,7 +320,7 @@ bool FActorDescTreeItem::GetPinnedState() const
 {
 	if (ActorDescHandle.IsValid() && ActorDescHandle->GetContainer())
 	{
-		UWorldPartition* WorldPartition = ActorDescHandle->GetContainer()->GetWorld()->GetWorldPartition();
+		UWorldPartition* WorldPartition = ActorDescHandle->GetContainer()->GetWorldPartition();
 		return WorldPartition ? WorldPartition->IsActorPinned(GetGuid()) : false;
 	}
 	return false;
diff --git a/Engine/Source/Editor/SceneOutliner/Private/SceneOutlinerModule.cpp b/Engine/Source/Editor/SceneOutliner/Private/SceneOutlinerModule.cpp
--- a/Engine/Source/Editor/SceneOutliner/Private/SceneOutlinerModule.cpp
+++ b/Engine/Source/Editor/SceneOutliner/Private/SceneOutlinerModule.cpp
@@ -25,7 +25,7 @@
 #include "WorldTreeItem.h"
 #include "WorldPartition/ActorDescContainer.h"
 #include "WorldPartition/DataLayer/DataLayerInstance.h"
-#include "WorldPartition/DataLayer/DataLayerSubsystem.h"
+#include "WorldPartition/DataLayer/DataLayerManager.h"
 #include "WorldPartition/ContentBundle/ContentBundleEngineSubsystem.h"
 #include "WorldPartition/ContentBundle/ContentBundleDescriptor.h"
 #include "LevelInstance/LevelInstanceSubsystem.h"
@@ -302,12 +302,12 @@ void FSceneOutlinerModule::CreateActorInfoColumns(FSceneOutlinerInitializationOp
 				{
 					if (const UActorDescContainer* ActorDescContainer = ActorDescItem->ActorDescHandle.Container.Get())
 					{
-						const UWorld* World = ActorDescContainer->GetWorld();
-						if (const UDataLayerSubsystem* DataLayerSubsystem = UWorld::GetSubsystem<UDataLayerSubsystem>(World))
+						const UWorld* OwningWorld = ActorDescContainer->GetWorldPartition()->GetWorld();
+						if (const UDataLayerManager* DataLayerManager = UDataLayerManager::GetDataLayerManager(OwningWorld))
 						{
 							TSet<const UDataLayerInstance*> DataLayerInstances;
-							DataLayerInstances.Append(DataLayerSubsystem->GetDataLayerInstances(ActorDesc->GetDataLayerInstanceNames()));
-							if (ULevelInstanceSubsystem* LevelInstanceSubsystem = UWorld::GetSubsystem<ULevelInstanceSubsystem>(World))
+							DataLayerInstances.Append(DataLayerManager->GetDataLayerInstances(ActorDesc->GetDataLayerInstanceNames()));
+							if (ULevelInstanceSubsystem* LevelInstanceSubsystem = UWorld::GetSubsystem<ULevelInstanceSubsystem>(OwningWorld))
 							{
 								UWorld* OuterWorld = ActorDescContainer->GetTypedOuter<UWorld>();
 								// Add parent container Data Layer Instances
diff --git a/Engine/Source/Editor/UnrealEd/Private/WorldPartition/WorldPartitionBuilder.cpp b/Engine/Source/Editor/UnrealEd/Private/WorldPartition/WorldPartitionBuilder.cpp
--- a/Engine/Source/Editor/UnrealEd/Private/WorldPartition/WorldPartitionBuilder.cpp
+++ b/Engine/Source/Editor/UnrealEd/Private/WorldPartition/WorldPartitionBuilder.cpp
@@ -12,7 +12,7 @@
 #include "WorldPartition/WorldPartition.h"
 #include "WorldPartition/WorldPartitionHelpers.h"
 #include "WorldPartition/WorldPartitionSubsystem.h"
-#include "WorldPartition/DataLayer/DataLayerSubsystem.h"
+#include "WorldPartition/DataLayer/DataLayerManager.h"
 #include "WorldPartition/DataLayer/DataLayerInstance.h"
 #include "WorldPartition/LoaderAdapter/LoaderAdapterShape.h"
 #include "LevelInstance/LevelInstanceSubsystem.h"
@@ -181,11 +181,11 @@ bool UWorldPartitionBuilder::Run(UWorld* World, FPackageSourceControlHelper& Pac
 	UWorldPartition* WorldPartition = World->GetWorldPartition();
 
 	// Properly Setup DataLayers for Builder
-	UDataLayerSubsystem* DataLayerSubsystem = UWorld::GetSubsystem<UDataLayerSubsystem>(World);
+	UDataLayerManager* DataLayerManager = WorldPartition->GetDataLayerManager();
 	
 	// Load Data Layers
 	bool bUpdateEditorCells = false;
-	DataLayerSubsystem->ForEachDataLayer([&bUpdateEditorCells, this](UDataLayerInstance* DataLayer)
+	DataLayerManager->ForEachDataLayerInstance([&bUpdateEditorCells, this](UDataLayerInstance* DataLayer)
 	{
 		const FName DataLayerShortName(DataLayer->GetDataLayerShortName());
 
diff --git a/Engine/Source/Editor/UnrealEd/Private/WorldPartition/WorldPartitionMiniMapBuilder.cpp b/Engine/Source/Editor/UnrealEd/Private/WorldPartition/WorldPartitionMiniMapBuilder.cpp
--- a/Engine/Source/Editor/UnrealEd/Private/WorldPartition/WorldPartitionMiniMapBuilder.cpp
+++ b/Engine/Source/Editor/UnrealEd/Private/WorldPartition/WorldPartitionMiniMapBuilder.cpp
@@ -19,7 +19,7 @@
 #include "WorldPartition/ActorDescContainer.h"
 #include "WorldPartition/DataLayer/ActorDataLayer.h"
 #include "WorldPartition/DataLayer/DataLayerInstance.h"
-#include "WorldPartition/DataLayer/DataLayerSubsystem.h"
+#include "WorldPartition/DataLayer/DataLayerManager.h"
 #include "WorldPartition/HLOD/HLODActor.h"
 #include "WorldPartition/HLOD/HLODActorDesc.h"
 #include "WorldPartition/WorldPartition.h"
@@ -82,10 +82,10 @@ bool UWorldPartitionMiniMapBuilder::PreRun(UWorld* World, FPackageSourceControlH
 
 	// Gather excluded data layers
 	{
-		UDataLayerSubsystem* DataLayerSubSystem = UWorld::GetSubsystem<UDataLayerSubsystem>(World);
+		UDataLayerManager* DataLayerManager = UDataLayerManager::GetDataLayerManager(World);
 		for (const FActorDataLayer& ActorDataLayer : WorldMiniMap->ExcludedDataLayers)
 		{
-			const UDataLayerInstance* DataLayerInstance = DataLayerSubSystem->GetDataLayerInstance(ActorDataLayer.Name);
+			const UDataLayerInstance* DataLayerInstance = DataLayerManager->GetDataLayerInstance(ActorDataLayer.Name);
 
 			if (DataLayerInstance != nullptr)
 			{
diff --git a/Engine/Source/Editor/UnrealEd/Private/WorldPartition/WorldPartitionNavigationDataBuilder.cpp b/Engine/Source/Editor/UnrealEd/Private/WorldPartition/WorldPartitionNavigationDataBuilder.cpp
--- a/Engine/Source/Editor/UnrealEd/Private/WorldPartition/WorldPartitionNavigationDataBuilder.cpp
+++ b/Engine/Source/Editor/UnrealEd/Private/WorldPartition/WorldPartitionNavigationDataBuilder.cpp
@@ -16,7 +16,7 @@
 #include "WorldPartition/WorldPartition.h"
 #include "WorldPartition/WorldPartitionSubsystem.h"
 #include "WorldPartition/DataLayer/DataLayerAsset.h"
-#include "WorldPartition/DataLayer/DataLayerSubsystem.h"
+#include "WorldPartition/DataLayer/DataLayerManager.h"
 #include "WorldPartition/NavigationData/NavigationDataChunkActor.h"
 
 DEFINE_LOG_CATEGORY_STATIC(LogWorldPartitionNavigationDataBuilder, Log, All);
@@ -29,12 +29,12 @@ UWorldPartitionNavigationDataBuilder::UWorldPartitionNavigationDataBuilder(const
 bool UWorldPartitionNavigationDataBuilder::PreRun(UWorld* World, FPackageSourceControlHelper& PackageHelper)
 {
 	// Set runtime data layer to be included in the base navmesh generation.
-	UDataLayerSubsystem* DataLayerSubsystem = UWorld::GetSubsystem<UDataLayerSubsystem>(World);
+	UDataLayerManager* DataLayerManager = UDataLayerManager::GetDataLayerManager(World);
 	for (const TObjectPtr<UDataLayerAsset> DataLayer : World->GetWorldSettings()->BaseNavmeshDataLayers)
 	{
 		if (DataLayer != nullptr)
 		{
-			const UDataLayerInstance* DataLayerInstance = DataLayerSubsystem->GetDataLayerInstance(DataLayer);
+			const UDataLayerInstance* DataLayerInstance = DataLayerManager->GetDataLayerInstance(DataLayer);
 			if (DataLayerInstance == nullptr)
 			{
 				UE_LOG(LogWorldPartitionNavigationDataBuilder, Error, TEXT("Missing UDataLayerInstance for %s."), *DataLayer->GetName());
diff --git a/Engine/Source/Runtime/Engine/Classes/Engine/World.h b/Engine/Source/Runtime/Engine/Classes/Engine/World.h
--- a/Engine/Source/Runtime/Engine/Classes/Engine/World.h
+++ b/Engine/Source/Runtime/Engine/Classes/Engine/World.h
@@ -52,6 +52,7 @@ class AServerStreamingLevelsVisibility;
 class AWorldDataLayers;
 class AWorldSettings;
 class UWorldPartition;
+class UDataLayerManager;
 class Error;
 class FConstPawnIterator;
 class FRegisterComponentContext;
@@ -2689,6 +2690,14 @@ public:
 	 */
 	UWorldPartition* GetWorldPartition() const;
 
+	/**
+	 * Returns the UDataLayerManager associated with this world.
+	 *
+	 * @return UDataLayerManager object associated with this world
+	 */
+	UFUNCTION(BlueprintCallable, Category = "Utilities|WorldPartition")
+	UDataLayerManager* GetDataLayerManager() const;
+
 	/**
 	* Returns true if world contains an associated UWorldPartition object.
 	*/
diff --git a/Engine/Source/Runtime/Engine/Private/Actor.cpp b/Engine/Source/Runtime/Engine/Private/Actor.cpp
--- a/Engine/Source/Runtime/Engine/Private/Actor.cpp
+++ b/Engine/Source/Runtime/Engine/Private/Actor.cpp
@@ -65,7 +65,7 @@
 #include "WorldPartition/DataLayer/WorldDataLayers.h"
 #include "WorldPartition/WorldPartitionActorDescUtils.h"
 #include "WorldPartition/DataLayer/DataLayerAsset.h"
-#include "WorldPartition/DataLayer/DataLayerSubsystem.h"
+#include "WorldPartition/DataLayer/DataLayerManager.h"
 #include "WorldPartition/ContentBundle/ContentBundlePaths.h"
 
 DEFINE_LOG_CATEGORY(LogActor);
@@ -6135,8 +6135,8 @@ TArray<const UDataLayerInstance*> AActor::GetDataLayerInstances() const
 }
 
 // Returns all valid DataLayerInstances for this actor including those inherited from their parent level instance actor.
-// If bUseLevelContext is true, the actor level will be passed down to the DataLayerSubsystem which will use it 
-// to narrow down the resolving of valid datalayers for this particular level.
+// If bUseLevelContext is true, the actor level will be used to find the associated DataLayerManager which will be used 
+// to resolve valid datalayers for this particular level.
 TArray<const UDataLayerInstance*> AActor::GetDataLayerInstancesInternal(bool bUseLevelContext, bool bIncludeParentDataLayers) const
 {
 	if (UseWorldPartitionRuntimeCellDataLayers())
@@ -6149,13 +6149,12 @@ TArray<const UDataLayerInstance*> AActor::GetDataLayerInstancesInternal(bool bUs
 	if (SupportsDataLayer())
 	{
 		TArray<const UDataLayerInstance*> DataLayerInstances;
-		ULevel* OuterLevel = bUseLevelContext ? GetLevel() : nullptr;
-		if (UDataLayerSubsystem* DataLayerSubsystem = UWorld::GetSubsystem<UDataLayerSubsystem>(GetWorld()))
+		if (UDataLayerManager* DataLayerManager = bUseLevelContext ? UDataLayerManager::GetDataLayerManager(this) : UDataLayerManager::GetDataLayerManager(GetWorld()))
 		{
-			DataLayerInstances += DataLayerSubsystem->GetDataLayerInstances(DataLayerAssets, OuterLevel);
+			DataLayerInstances += DataLayerManager->GetDataLayerInstances(DataLayerAssets);
 
 			PRAGMA_DISABLE_DEPRECATION_WARNINGS
-			DataLayerInstances += DataLayerSubsystem->GetDataLayerInstances(DataLayers, OuterLevel);
+			DataLayerInstances += DataLayerManager->GetDataLayerInstances(DataLayers);
 			PRAGMA_ENABLE_DEPRECATION_WARNINGS
 		}
 
diff --git a/Engine/Source/Runtime/Engine/Private/ActorEditor.cpp b/Engine/Source/Runtime/Engine/Private/ActorEditor.cpp
--- a/Engine/Source/Runtime/Engine/Private/ActorEditor.cpp
+++ b/Engine/Source/Runtime/Engine/Private/ActorEditor.cpp
@@ -12,7 +12,7 @@
 #include "WorldPartition/DataLayer/WorldDataLayers.h"
 #include "WorldPartition/WorldPartition.h"
 #include "WorldPartition/DataLayer/DataLayer.h"
-#include "WorldPartition/DataLayer/DataLayerSubsystem.h"
+#include "WorldPartition/DataLayer/DataLayerManager.h"
 #include "EditorSupportDelegates.h"
 #include "Logging/MessageLog.h"
 #include "Misc/UObjectToken.h"
@@ -1632,91 +1632,96 @@ TArray<FName> AActor::GetDataLayerInstanceNames() const
 }
 
 TArray<const UDataLayerInstance*> AActor::GetDataLayerInstancesForLevel() const
-	{
+{
 	const bool bUseLevelContext = true;
 	return GetDataLayerInstancesInternal(bUseLevelContext);
 }
 
 void AActor::FixupDataLayers(bool bRevertChangesOnLockedDataLayer /*= false*/)
 {
-	if (!GetPackage()->HasAnyPackageFlags(PKG_PlayInEditor))
+	if (IsTemplate())
 	{
-		if (!SupportsDataLayer())
-		{
-			DataLayerAssets.Empty();
-			DataLayers.Empty();
-			return;
-		}
-
-		if (UDataLayerSubsystem* DataLayerSubsystem = UWorld::GetSubsystem<UDataLayerSubsystem>(GetWorld()))
-		{
-			if (!DataLayerSubsystem->CanResolveDataLayers())
-			{
-				return;
-			}
+		return;
+	}
 
-			ULevel* Level = GetLevel();
+	if (!SupportsDataLayer())
+	{
+		DataLayerAssets.Empty();
+		DataLayers.Empty();
+		return;
+	}
 
-				if (bRevertChangesOnLockedDataLayer)
-				{
-					// Since it's not possible to prevent changes of particular elements of an array, rollback change on locked DataLayers.
-					TSet<const UDataLayerAsset*> PreEdit(PreEditChangeDataLayers);
-					TSet<const UDataLayerAsset*> PostEdit(DataLayerAssets);
+	// Don't fixup in game world
+	UWorld* World = GetWorld();
+	if ((GetPackage()->HasAnyPackageFlags(PKG_PlayInEditor)) ||
+		(World && (World->IsGameWorld() || (World->WorldType == EWorldType::Inactive))))
+	{ 
+		return;
+	}
 
-				auto DifferenceContainsLockedDataLayers = [DataLayerSubsystem, Level](const TSet<const UDataLayerAsset*>& A, const TSet<const UDataLayerAsset*>& B)
-					{
-						TSet<const UDataLayerAsset*> Diff = A.Difference(B);
-						for (const UDataLayerAsset* DataLayerAsset : Diff)
-						{
-						// We pass Actor Level when resolving the DataLayerInstance as we do the fixup relative to this level
-						const UDataLayerInstance* DataLayerInstance = DataLayerSubsystem->GetDataLayerInstance(DataLayerAsset, Level);
-							if (DataLayerInstance && DataLayerInstance->IsLocked())
-							{
-								return true;
-							}
-						}
-						return false;
-					};
-					
-					if (DifferenceContainsLockedDataLayers(PreEdit, PostEdit) || 
-						DifferenceContainsLockedDataLayers(PostEdit, PreEdit))
-					{
-						DataLayerAssets = PreEditChangeDataLayers;
-					}
-				}
+	// Use Actor's DataLayerManager since the fixup is relative to this level
+	UDataLayerManager* DataLayerManager = UDataLayerManager::GetDataLayerManager(this);
+	if (!DataLayerManager || !DataLayerManager->CanResolveDataLayers())
+	{
+		return;
+	}
+	
+	if (bRevertChangesOnLockedDataLayer)
+	{
+		// Since it's not possible to prevent changes of particular elements of an array, rollback change on locked DataLayers.
+		TSet<const UDataLayerAsset*> PreEdit(PreEditChangeDataLayers);
+		TSet<const UDataLayerAsset*> PostEdit(DataLayerAssets);
 
-				PRAGMA_DISABLE_DEPRECATION_WARNINGS
-			auto CleanupDataLayers = [this, DataLayerSubsystem, Level](auto& DataLayerArray)
+		auto DifferenceContainsLockedDataLayers = [DataLayerManager](const TSet<const UDataLayerAsset*>& A, const TSet<const UDataLayerAsset*>& B)
+		{
+			TSet<const UDataLayerAsset*> Diff = A.Difference(B);
+			for (const UDataLayerAsset* DataLayerAsset : Diff)
+			{
+				const UDataLayerInstance* DataLayerInstance = DataLayerManager->GetDataLayerInstance(DataLayerAsset);
+				if (DataLayerInstance && DataLayerInstance->IsLocked())
 				{
-					using ArrayType = typename TRemoveReference<decltype(DataLayerArray)>::Type;
+					return true;
+				}
+			}
+			return false;
+		};
 					
-					ArrayType ExistingDataLayer;
-					for (int32 Index = 0; Index < DataLayerArray.Num();)
-					{
-					// We pass Actor Level when resolving the DataLayerInstance as we do the fixup relative to this level
-						auto& DataLayer = DataLayerArray[Index];
-					if (!DataLayerSubsystem->GetDataLayerInstance(DataLayer, Level) || ExistingDataLayer.Contains(DataLayer))
-						{
-							DataLayerArray.RemoveAtSwap(Index);
-						}
-						else
-						{
-							ExistingDataLayer.Add(DataLayer);
-							++Index;
-						}
-					}
-				};
+		if (DifferenceContainsLockedDataLayers(PreEdit, PostEdit) || 
+			DifferenceContainsLockedDataLayers(PostEdit, PreEdit))
+		{
+			DataLayerAssets = PreEditChangeDataLayers;
+		}
+	}
 
-				// Only invalidate DataLayerAssets on cook. In Editor we want to be able to re-resolve if the asset gets readded to the WorldDataLayers actor
-				if (IsRunningCookCommandlet())
-				{
-					CleanupDataLayers(DataLayerAssets);
-				}
-				CleanupDataLayers(DataLayers);
-				PRAGMA_ENABLE_DEPRECATION_WARNINGS
+	PRAGMA_DISABLE_DEPRECATION_WARNINGS
+	auto CleanupDataLayers = [this, DataLayerManager](auto& DataLayerArray)
+	{
+		using ArrayType = typename TRemoveReference<decltype(DataLayerArray)>::Type;
+					
+		ArrayType ExistingDataLayer;
+		for (int32 Index = 0; Index < DataLayerArray.Num();)
+		{
+			auto& DataLayer = DataLayerArray[Index];
+			if (!DataLayerManager->GetDataLayerInstance(DataLayer) || ExistingDataLayer.Contains(DataLayer))
+			{
+				DataLayerArray.RemoveAtSwap(Index);
+			}
+			else
+			{
+				ExistingDataLayer.Add(DataLayer);
+				++Index;
 			}
 		}
+	};
+
+	// Only invalidate DataLayerAssets on cook. In Editor we want to be able to re-resolve if the asset gets readded to the WorldDataLayers actor
+	if (IsRunningCookCommandlet())
+	{
+		CleanupDataLayers(DataLayerAssets);
 	}
+	CleanupDataLayers(DataLayers);
+	PRAGMA_ENABLE_DEPRECATION_WARNINGS
+}
 
 bool AActor::IsPropertyChangedAffectingDataLayers(FPropertyChangedEvent& PropertyChangedEvent) const
 {
diff --git a/Engine/Source/Runtime/Engine/Private/ActorPartition/ActorPartitionSubsystem.cpp b/Engine/Source/Runtime/Engine/Private/ActorPartition/ActorPartitionSubsystem.cpp
--- a/Engine/Source/Runtime/Engine/Private/ActorPartition/ActorPartitionSubsystem.cpp
+++ b/Engine/Source/Runtime/Engine/Private/ActorPartition/ActorPartitionSubsystem.cpp
@@ -8,9 +8,12 @@
 #include "WorldPartition/DataLayer/DataLayerEditorContext.h"
 #include "WorldPartition/WorldPartitionSubsystem.h"
 #include "WorldPartition/WorldPartition.h"
-#include "WorldPartition/DataLayer/DataLayerSubsystem.h"
 #include "EngineUtils.h"
 
+#if WITH_EDITOR
+#include "WorldPartition/DataLayer/DataLayerManager.h"
+#endif
+
 #include UE_INLINE_GENERATED_CPP_BY_NAME(ActorPartitionSubsystem)
 
 DEFINE_LOG_CATEGORY_STATIC(LogActorPartitionSubsystem, All, All);
@@ -418,12 +421,13 @@ APartitionActor* UActorPartitionSubsystem::GetActor(const TSubclassOf<APartition
 	const uint32 GridSize = InGridSize > 0 ? InGridSize : InActorClass->GetDefaultObject<APartitionActor>()->GetDefaultGridSize(GetWorld());
 	
 	UWorld* World = GetWorld();
+	const UDataLayerManager* DataLayerManager = UDataLayerManager::GetDataLayerManager(World);
 	
-	auto WrappedInActorCreated = [World, InActorCreated](APartitionActor* PartitionActor)
+	auto WrappedInActorCreated = [DataLayerManager, InActorCreated](APartitionActor* PartitionActor)
 	{
-		if(UDataLayerSubsystem* DataLayerSubsystem = World->GetSubsystem<UDataLayerSubsystem>())
+		if(DataLayerManager)
 		{
-			for (UDataLayerInstance* DataLayer : DataLayerSubsystem->GetActorEditorContextDataLayers())
+			for (UDataLayerInstance* DataLayer : DataLayerManager->GetActorEditorContextDataLayers())
 			{
 				PartitionActor->AddDataLayer(DataLayer);
 			}
@@ -431,8 +435,7 @@ APartitionActor* UActorPartitionSubsystem::GetActor(const TSubclassOf<APartition
 		InActorCreated(PartitionActor);
 	};
 
-	UDataLayerSubsystem* DataLayerSubsystem = World->GetSubsystem<UDataLayerSubsystem>();
-	FActorPartitionIdentifier ActorPartitionId(InActorClass, InGuid, DataLayerSubsystem ? DataLayerSubsystem->GetDataLayerEditorContextHash() : FDataLayerEditorContext::EmptyHash);
+	FActorPartitionIdentifier ActorPartitionId(InActorClass, InGuid, DataLayerManager ? DataLayerManager->GetDataLayerEditorContextHash() : FDataLayerEditorContext::EmptyHash);
 	TMap<FActorPartitionIdentifier, TWeakObjectPtr<APartitionActor>>* ActorsPerId = PartitionedActors.Find(InCellCoords);
 	APartitionActor* FoundActor = nullptr;
 	if (!ActorsPerId)
diff --git a/Engine/Source/Runtime/Engine/Private/LevelInstance/LevelInstanceSubsystem.cpp b/Engine/Source/Runtime/Engine/Private/LevelInstance/LevelInstanceSubsystem.cpp
--- a/Engine/Source/Runtime/Engine/Private/LevelInstance/LevelInstanceSubsystem.cpp
+++ b/Engine/Source/Runtime/Engine/Private/LevelInstance/LevelInstanceSubsystem.cpp
@@ -21,7 +21,7 @@
 #include "LevelInstance/LevelInstanceEditorPivotActor.h"
 #include "WorldPartition/WorldPartition.h"
 #include "WorldPartition/WorldPartitionSubsystem.h"
-#include "WorldPartition/DataLayer/DataLayerSubsystem.h"
+#include "WorldPartition/DataLayer/DataLayerManager.h"
 #include "WorldPartition/DataLayer/DataLayerInstanceWithAsset.h"
 #include "WorldPartition/WorldPartitionMiniMap.h"
 #include "Misc/ScopedSlowTask.h"
@@ -1028,35 +1028,39 @@ ILevelInstanceInterface* ULevelInstanceSubsystem::CreateLevelInstanceFrom(const
 			if (bIsPartitioned)
 			{
 				UWorldPartition* WorldPartition = InLevel->GetWorldPartition();
+				// Validations
 				check(WorldPartition);
-
-				// Make sure that the world partition can be used by a Level Instance
 				check(WorldPartition->CanBeUsedByLevelInstance());
+				check(!WorldPartition->IsStreamingEnabled());
+				check(InLevel->IsUsingActorFolders());
 
 				// Make sure new level's AWorldDataLayers contains all the necessary Data Layer Instances before moving actors
-				if (UDataLayerSubsystem* DataLayerSubsystem = UWorld::GetSubsystem<UDataLayerSubsystem>(GetWorld()))
+				TSet<TObjectPtr<const UDataLayerAsset>> SourceDataLayerAssets;
+				for (AActor* ActorToMove : ActorsToMove)
 				{
-					TSet<TObjectPtr<const UDataLayerAsset>> SourceDataLayerAssets;
-					for (AActor* ActorToMove : ActorsToMove)
+					if (UDataLayerManager* DataLayerManager = UDataLayerManager::GetDataLayerManager(ActorToMove))
 					{
 						// Use the raw asset list as we don't want parent DataLayers
-						SourceDataLayerAssets.Append(ActorToMove->GetDataLayerAssets());
+						for (const TObjectPtr<const UDataLayerAsset>& DataLayerAsset : ActorToMove->GetDataLayerAssets())
+						{
+							// Validate that there's a valid Data Layer Instance for this asset in the source level
+							if (DataLayerManager->GetDataLayerInstanceFromAsset(DataLayerAsset))
+							{
+								SourceDataLayerAssets.Add(DataLayerAsset);
+							}
+						}
 					}
+				}
+
+				if (!SourceDataLayerAssets.IsEmpty())
+				{
 					AWorldDataLayers* WorldDataLayers = InLevel->GetWorldDataLayers();
 					check(WorldDataLayers);
-
 					for (const UDataLayerAsset* SourceDataLayerAsset : SourceDataLayerAssets)
 					{
-						if (UDataLayerInstanceWithAsset* SourceDataLayerInstance = Cast<UDataLayerInstanceWithAsset>(DataLayerSubsystem->GetDataLayerInstanceFromAsset(SourceDataLayerAsset)))
-						{
-							UDataLayerInstanceWithAsset* DataLayerInstanceWithAsset = WorldDataLayers->CreateDataLayer<UDataLayerInstanceWithAsset>(SourceDataLayerAsset);
-						}
+						WorldDataLayers->CreateDataLayer<UDataLayerInstanceWithAsset>(SourceDataLayerAsset);
 					}
 				}
-
-				// Validation
-				check(InLevel->IsUsingActorFolders());
-				check(!WorldPartition->IsStreamingEnabled());
 			}
 		}));
 
@@ -1506,7 +1510,7 @@ UActorDescContainer* ULevelInstanceSubsystem::FActorDescContainerInstanceManager
 		ExistingContainerInstance->UpdateBounds();
 	}
 
-	check(ActorDescContainer->GetWorld() == InWorld);
+	check(ActorDescContainer->IsTemplateContainer());
 	return ActorDescContainer;
 }
 
@@ -1965,11 +1969,15 @@ bool ULevelInstanceSubsystem::EditLevelInstanceInternal(ILevelInstanceInterface*
 	LevelInstancesToLoadOrUpdate.Remove(LevelInstance);
 	// Unload right away
 	UnloadLevelInstance(LevelInstance->GetLevelInstanceID());
-		
+	
+	// When editing a Level Instance, push a new empty actor editor context
+	UActorEditorContextSubsystem::Get()->PushContext();
+
 	// Load Edit LevelInstance level
 	ULevelStreamingLevelInstanceEditor* LevelStreaming = ULevelStreamingLevelInstanceEditor::Load(LevelInstance);
 	if (!LevelStreaming)
 	{
+		UActorEditorContextSubsystem::Get()->PopContext();
 		LevelInstance->LoadLevelInstance();
 		return false;
 	}
@@ -2010,9 +2018,6 @@ bool ULevelInstanceSubsystem::EditLevelInstanceInternal(ILevelInstanceInterface*
 	// Edit can't be undone
 	GEditor->ResetTransaction(LOCTEXT("LevelInstanceEditResetTrans", "Edit Level Instance"));
 
-	// When editing a Level Instance, push a new empty actor editor context
-	UActorEditorContextSubsystem::Get()->PushContext();
-
 	ResetLoadersForWorldAsset(LevelInstance->GetWorldAsset().GetLongPackageName());
 
 	if (UWorldPartition* WorldPartition = LevelInstanceActor->GetWorld()->GetWorldPartition(); WorldPartition && WorldPartition->IsMainWorldPartition())
diff --git a/Engine/Source/Runtime/Engine/Private/World.cpp b/Engine/Source/Runtime/Engine/Private/World.cpp
--- a/Engine/Source/Runtime/Engine/Private/World.cpp
+++ b/Engine/Source/Runtime/Engine/Private/World.cpp
@@ -50,6 +50,7 @@
 #include "Engine/Console.h"
 #include "Engine/WorldComposition.h"
 #include "WorldPartition/WorldPartition.h"
+#include "WorldPartition/DataLayer/DataLayerManager.h"
 #include "WorldPartition/DataLayer/WorldDataLayers.h"
 #include "GameFramework/GameNetworkManager.h"
 #include "GameFramework/CharacterMovementComponent.h"
@@ -5606,6 +5607,12 @@ UWorldPartition* UWorld::GetWorldPartition() const
 	return WorldSettings ? WorldSettings->GetWorldPartition() : nullptr;
 }
 
+UDataLayerManager* UWorld::GetDataLayerManager() const
+{
+	UWorldPartition* WorldPartition = GetWorldPartition();
+	return WorldPartition ? WorldPartition->GetDataLayerManager() : nullptr;
+}
+
 UModel* UWorld::GetModel() const
 {
 #if !WITH_EDITORONLY_DATA
diff --git a/Engine/Source/Runtime/Engine/Private/WorldPartition/ActorDescContainer.cpp b/Engine/Source/Runtime/Engine/Private/WorldPartition/ActorDescContainer.cpp
--- a/Engine/Source/Runtime/Engine/Private/WorldPartition/ActorDescContainer.cpp
+++ b/Engine/Source/Runtime/Engine/Private/WorldPartition/ActorDescContainer.cpp
@@ -5,21 +5,21 @@
 #include UE_INLINE_GENERATED_CPP_BY_NAME(ActorDescContainer)
 
 #if WITH_EDITOR
-#include "AssetRegistry/ARFilter.h"
 #include "Editor.h"
-#include "AssetRegistry/AssetRegistryModule.h"
 #include "Engine/Level.h"
-#include "WorldPartition/WorldPartition.h"
+#include "UObject/ObjectSaveContext.h"
+#include "AssetRegistry/ARFilter.h"
+#include "AssetRegistry/AssetRegistryModule.h"
 #include "Modules/ModuleManager.h"
+#include "WorldPartition/WorldPartition.h"
 #include "WorldPartition/WorldPartitionActorDescUtils.h"
-#include "UObject/ObjectSaveContext.h"
+#include "WorldPartition/DataLayer/DataLayerManager.h"
 
 UActorDescContainer::FActorDescContainerInitializeDelegate UActorDescContainer::OnActorDescContainerInitialized;
 #endif
 
 UActorDescContainer::UActorDescContainer(const FObjectInitializer& ObjectInitializer)
 	: Super(ObjectInitializer)
-	, World(nullptr)
 #if WITH_EDITOR
 	, bContainerInitialized(false)
 #endif
@@ -34,8 +34,10 @@ void UActorDescContainer::Initialize(const FInitializeParams& InitParams)
 {
 	TRACE_CPUPROFILER_EVENT_SCOPE(UActorDescContainer::Initialize);
 
-	check(!World || World == InitParams.World);
-	World = InitParams.World;
+	// @todo_ow: We need to pass the world context to AddActorDescriptor for FLevelInstanceActorDesc::RegisterContainerInstance to resolve the ULevelInstanceSubsystem.
+	// A better solution would be for ActorDescContainers to always be outered to an OwningWorldPartition (with the downside of not sharing between 2 instanced WorldPartition).
+	// With this, we could always find the owning WorldPartition and of course the owning world (GetOwningWorldPartition()->GetWorld()).
+	UWorld* OwningWorld = InitParams.World;
 
 #if WITH_EDITOR
 	check(!bContainerInitialized);
@@ -70,7 +72,7 @@ void UActorDescContainer::Initialize(const FInitializeParams& InitParams)
 
 			if (ActorDesc.IsValid() && ActorDesc->GetNativeClass().IsValid() && (!InitParams.FilterActorDesc || InitParams.FilterActorDesc(ActorDesc.Get())))
 			{
-				AddActorDescriptor(ActorDesc.Release());
+				AddActorDescriptor(ActorDesc.Release(), OwningWorld);
 			}
 			else
 			{
@@ -92,8 +94,6 @@ void UActorDescContainer::Initialize(const FInitializeParams& InitParams)
 
 void UActorDescContainer::Update()
 {
-	check(World);
-
 #if WITH_EDITOR
 	check(bContainerInitialized);
 	TArray<FAssetData> Assets;
@@ -111,6 +111,9 @@ void UActorDescContainer::Update()
 
 	AssetRegistry.GetAssets(Filter, Assets);
 
+	UWorld* OwningWorld = GetWorldPartition() ? GetWorldPartition()->GetWorld() : nullptr;
+	check(OwningWorld);
+
 	TSet<FGuid> ActorGuids;
 	for (const FAssetData& Asset : Assets)
 	{
@@ -132,7 +135,7 @@ void UActorDescContainer::Update()
 			else
 			{
 				FWorldPartitionActorDesc* ActorDescPtr = NewActorDesc.Release();
-				AddActorDescriptor(ActorDescPtr);
+				AddActorDescriptor(ActorDescPtr, OwningWorld);
 				OnActorDescAdded(ActorDescPtr);
 			}
 		}
@@ -175,16 +178,6 @@ void UActorDescContainer::Uninitialize()
 		ActorDescPtr.Reset();
 	}
 #endif
-	World = nullptr;
-}
-
-UWorld* UActorDescContainer::GetWorld() const
-{
-	if (World)
-	{
-		return World;
-	}
-	return Super::GetWorld();
 }
 
 void UActorDescContainer::BeginDestroy()
@@ -213,18 +206,36 @@ bool UActorDescContainer::IsActorDescHandled(const AActor* Actor) const
 
 bool UActorDescContainer::IsMainPartitionContainer() const
 {
-	return GetWorld() != nullptr && GetWorld()->GetWorldPartition()->GetActorDescContainer() == this;
+	return GetWorldPartition() && GetWorldPartition()->GetActorDescContainer() == this;
+}
+
+bool UActorDescContainer::IsTemplateContainer() const
+{
+	bool bIsTemplate = GetOuter() == GetTransientPackage();
+	check(bIsTemplate || GetWorldPartition());
+	return bIsTemplate;
+}
+
+UWorldPartition* UActorDescContainer::GetWorldPartition() const
+{
+	UWorldPartition* OuterWorldPartition = GetTypedOuter<UWorldPartition>();
+	check(OuterWorldPartition || IsTemplateContainer());
+	return OuterWorldPartition;
 }
 
-void UActorDescContainer::AddActorDescriptor(FWorldPartitionActorDesc* ActorDesc)
+void UActorDescContainer::AddActorDescriptor(FWorldPartitionActorDesc* ActorDesc, UWorld* InWorldContext)
 {
-	FActorDescList::AddActorDescriptor(ActorDesc);
-	ActorDesc->SetContainer(this);
+	FActorDescList::AddActorDescriptor(ActorDesc, InWorldContext);
+	ActorDesc->SetContainer(this, InWorldContext);
+	if (const UDataLayerManager* DataLayerManager = UDataLayerManager::GetDataLayerManager(InWorldContext))
+	{
+		DataLayerManager->ResolveActorDescDataLayers(ActorDesc);
+	}
 }
 
 void UActorDescContainer::RemoveActorDescriptor(FWorldPartitionActorDesc* ActorDesc)
 {
-	ActorDesc->SetContainer(nullptr);
+	ActorDesc->SetContainer(nullptr, nullptr);
 	FActorDescList::RemoveActorDescriptor(ActorDesc);
 }
 
@@ -323,7 +334,10 @@ bool UActorDescContainer::ShouldRegisterDelegates()
 	UWorld* OuterWorld = GetTypedOuter<UWorld>();
 	// No need to register delegates for level instances
 	bool bIsInstance = OuterWorld && OuterWorld->IsInstanced() && !OuterWorld->GetPackage()->HasAnyPackageFlags(PKG_NewlyCreated);
-	return GEditor && !IsTemplate() && World && !World->IsGameWorld() && !bIsInstance;
+	UWorld* OwningWorld = GetWorldPartition() ? GetWorldPartition()->GetWorld() : nullptr;
+	// Template container will always register
+	bool bShouldRegisterForWorld = (OwningWorld && !OwningWorld->IsGameWorld()) || (IsTemplateContainer());
+	return GEditor && !IsTemplate() && !bIsInstance && bShouldRegisterForWorld;
 }
 
 void UActorDescContainer::RegisterEditorDelegates()
@@ -350,7 +364,7 @@ void UActorDescContainer::OnActorDescAdded(FWorldPartitionActorDesc* NewActorDes
 {
 	OnActorDescAddedEvent.Broadcast(NewActorDesc);
 
-	if (UWorldPartition* WorldPartition = GetTypedOuter<UWorldPartition>())
+	if (UWorldPartition* WorldPartition = GetWorldPartition())
 	{
 		WorldPartition->OnActorDescAdded(NewActorDesc);
 	}
@@ -360,7 +374,7 @@ void UActorDescContainer::OnActorDescRemoved(FWorldPartitionActorDesc* ActorDesc
 {
 	OnActorDescRemovedEvent.Broadcast(ActorDesc);
 
-	if (UWorldPartition* WorldPartition = GetTypedOuter<UWorldPartition>())
+	if (UWorldPartition* WorldPartition = GetWorldPartition())
 	{
 		WorldPartition->OnActorDescRemoved(ActorDesc);
 	}
@@ -368,7 +382,7 @@ void UActorDescContainer::OnActorDescRemoved(FWorldPartitionActorDesc* ActorDesc
 
 void UActorDescContainer::OnActorDescUpdating(FWorldPartitionActorDesc* ActorDesc)
 {
-	if (UWorldPartition* WorldPartition = GetTypedOuter<UWorldPartition>())
+	if (UWorldPartition* WorldPartition = GetWorldPartition())
 	{
 		WorldPartition->OnActorDescUpdating(ActorDesc);
 	}
@@ -376,8 +390,12 @@ void UActorDescContainer::OnActorDescUpdating(FWorldPartitionActorDesc* ActorDes
 
 void UActorDescContainer::OnActorDescUpdated(FWorldPartitionActorDesc* ActorDesc)
 {
-	if (UWorldPartition* WorldPartition = GetTypedOuter<UWorldPartition>())
+	if (UWorldPartition* WorldPartition = GetWorldPartition())
 	{
+		if (const UDataLayerManager* DataLayerManager = UDataLayerManager::GetDataLayerManager(WorldPartition))
+		{
+			DataLayerManager->ResolveActorDescDataLayers(ActorDesc);
+		}
 		WorldPartition->OnActorDescUpdated(ActorDesc);
 	}
 }
@@ -386,7 +404,7 @@ const FLinkerInstancingContext* UActorDescContainer::GetInstancingContext() cons
 {
 	const FLinkerInstancingContext* InstancingContext = nullptr;
 
-	if (UWorldPartition* WorldPartition = GetTypedOuter<UWorldPartition>())
+	if (UWorldPartition* WorldPartition = GetWorldPartition())
 	{
 		WorldPartition->GetInstancingContext(InstancingContext);
 	}
@@ -396,7 +414,7 @@ const FLinkerInstancingContext* UActorDescContainer::GetInstancingContext() cons
 
 const FTransform& UActorDescContainer::GetInstanceTransform() const
 {
-	if (UWorldPartition* WorldPartition = GetTypedOuter<UWorldPartition>())
+	if (UWorldPartition* WorldPartition = GetWorldPartition())
 	{
 		return WorldPartition->GetInstanceTransform();
 	}
diff --git a/Engine/Source/Runtime/Engine/Private/WorldPartition/ActorDescList.cpp b/Engine/Source/Runtime/Engine/Private/WorldPartition/ActorDescList.cpp
--- a/Engine/Source/Runtime/Engine/Private/WorldPartition/ActorDescList.cpp
+++ b/Engine/Source/Runtime/Engine/Private/WorldPartition/ActorDescList.cpp
@@ -9,7 +9,7 @@ FWorldPartitionActorDesc* FActorDescList::AddActor(const AActor* InActor)
 	FWorldPartitionActorDesc* NewActorDesc = InActor->CreateActorDesc().Release();
 	check(NewActorDesc);
 
-	AddActorDescriptor(NewActorDesc);
+	AddActorDescriptor(NewActorDesc, InActor->GetWorld());
 
 	return NewActorDesc;
 }
@@ -67,7 +67,7 @@ void FActorDescList::Empty()
 	ActorDescList.Empty();
 }
 
-void FActorDescList::AddActorDescriptor(FWorldPartitionActorDesc* ActorDesc)
+void FActorDescList::AddActorDescriptor(FWorldPartitionActorDesc* ActorDesc, UWorld* InWorldContext)
 {
 	check(ActorDesc);
 
diff --git a/Engine/Source/Runtime/Engine/Private/WorldPartition/DataLayer/DataLayerEditorContext.cpp b/Engine/Source/Runtime/Engine/Private/WorldPartition/DataLayer/DataLayerEditorContext.cpp
--- a/Engine/Source/Runtime/Engine/Private/WorldPartition/DataLayer/DataLayerEditorContext.cpp
+++ b/Engine/Source/Runtime/Engine/Private/WorldPartition/DataLayer/DataLayerEditorContext.cpp
@@ -5,7 +5,7 @@
 
 #if WITH_EDITOR
 
-#include "WorldPartition/DataLayer/DataLayerSubsystem.h"
+#include "WorldPartition/DataLayer/DataLayerManager.h"
 
 /*
  * FDataLayerEditorContext
@@ -13,15 +13,15 @@
 FDataLayerEditorContext::FDataLayerEditorContext(UWorld* InWorld, const TArray<FName>& InDataLayerInstances)
 	: Hash(FDataLayerEditorContext::EmptyHash)
 {
-	UDataLayerSubsystem* DataLayerSubsystem = UWorld::GetSubsystem<UDataLayerSubsystem>(InWorld);
-	if (!DataLayerSubsystem)
+	UDataLayerManager* DataLayerManager = UDataLayerManager::GetDataLayerManager(InWorld);
+	if (!DataLayerManager)
 	{
 		return;
 	}
 
 	for (const FName& DataLayerInstanceName : InDataLayerInstances)
 	{
-		if (const UDataLayerInstance* DataLayerInstance = DataLayerSubsystem->GetDataLayerInstance(DataLayerInstanceName))
+		if (const UDataLayerInstance* DataLayerInstance = DataLayerManager->GetDataLayerInstance(DataLayerInstanceName))
 		{
 			DataLayerInstances.AddUnique(DataLayerInstance->GetDataLayerFName());
 		}
diff --git a/Engine/Source/Runtime/Engine/Private/WorldPartition/DataLayer/DataLayerInstance.cpp b/Engine/Source/Runtime/Engine/Private/WorldPartition/DataLayer/DataLayerInstance.cpp
--- a/Engine/Source/Runtime/Engine/Private/WorldPartition/DataLayer/DataLayerInstance.cpp
+++ b/Engine/Source/Runtime/Engine/Private/WorldPartition/DataLayer/DataLayerInstance.cpp
@@ -3,6 +3,7 @@
 #include "WorldPartition/DataLayer/DataLayerInstance.h"
 
 #include "UObject/UnrealType.h"
+#include "WorldPartition/WorldPartitionLog.h"
 #include "WorldPartition/DataLayer/DataLayerUtils.h"
 #include "WorldPartition/DataLayer/WorldDataLayers.h"
 #include "WorldPartition/ErrorHandling/WorldPartitionStreamingGenerationErrorHandler.h"
@@ -113,6 +114,11 @@ bool UDataLayerInstance::IsLocked() const
 	return IsRuntime() && !GetOuterAWorldDataLayers()->GetAllowRuntimeDataLayerEditing();
 }
 
+bool UDataLayerInstance::IsReadOnly() const
+{
+	return GetWorld()->IsGameWorld();
+}
+
 const TCHAR* UDataLayerInstance::GetDataLayerIconName() const
 {
 	return FDataLayerUtils::GetDataLayerIconName(GetType());
@@ -274,4 +280,25 @@ void UDataLayerInstance::AddChild(UDataLayerInstance* InDataLayer)
 	Children.Add(InDataLayer);
 }
 
+EDataLayerRuntimeState UDataLayerInstance::GetRuntimeState() const
+{
+	return GetOuterAWorldDataLayers()->GetDataLayerRuntimeStateByName(GetDataLayerFName());
+}
+
+EDataLayerRuntimeState UDataLayerInstance::GetEffectiveRuntimeState() const
+{
+	return GetOuterAWorldDataLayers()->GetDataLayerEffectiveRuntimeStateByName(GetDataLayerFName());
+}
+
+bool UDataLayerInstance::SetRuntimeState(EDataLayerRuntimeState InState, bool bInIsRecursive) const
+{
+	if (GetOuterAWorldDataLayers()->HasAuthority())
+	{
+		GetOuterAWorldDataLayers()->SetDataLayerRuntimeState(this, InState, bInIsRecursive);
+		return true;
+	}
+	UE_LOG(LogWorldPartition, Error, TEXT("SetDataLayerRuntimeState can only execute on authority"));
+	return false;
+}
+
 #undef LOCTEXT_NAMESPACE
