Properly hook Data Layers menu extenders and pass selection to delegate.

diff --git a/Engine/Source/Editor/DataLayerEditor/Private/DataLayer/DataLayerMode.cpp b/Engine/Source/Editor/DataLayerEditor/Private/DataLayer/DataLayerMode.cpp
--- a/Engine/Source/Editor/DataLayerEditor/Private/DataLayer/DataLayerMode.cpp
+++ b/Engine/Source/Editor/DataLayerEditor/Private/DataLayer/DataLayerMode.cpp
@@ -22,6 +22,7 @@
 #include "DataLayerActorTreeItem.h"
 #include "DataLayerHierarchy.h"
 #include "DataLayerTreeItem.h"
+#include "DataLayerEditorModule.h"
 #include "DataLayersActorDescTreeItem.h"
 #include "DragAndDrop/ActorDragDropOp.h"
 #include "DragAndDrop/CompositeDragDropOp.h"
@@ -1602,6 +1603,7 @@ TSharedPtr<SWidget> FDataLayerMode::CreateContextMenu()
 	RegisterContextMenu();
 
 	FSceneOutlinerItemSelection ItemSelection(SceneOutliner->GetSelection());
+	CacheSelectedItems(ItemSelection);
 
 	USceneOutlinerMenuContext* ContextObject = NewObject<USceneOutlinerMenuContext>();
 	ContextObject->SceneOutliner = StaticCastSharedRef<SSceneOutliner>(SceneOutliner->AsShared());
@@ -1612,6 +1614,25 @@ TSharedPtr<SWidget> FDataLayerMode::CreateContextMenu()
 	FName MenuName = DefaultContextMenuName;
 	SceneOutliner->GetSharedData().ModifyContextMenu.ExecuteIfBound(MenuName, Context);
 
+	// Add Extenders
+	FDataLayerEditorModule& DataLayerEditorModule = FModuleManager::LoadModuleChecked<FDataLayerEditorModule>("DataLayerEditor");
+	TArray<FDataLayerEditorModule::FDataLayersMenuExtender>& MenuExtenderDelegates = DataLayerEditorModule.GetAllDataLayersMenuExtenders();
+	TArray<TSharedPtr<FExtender>> Extenders;
+	if (!MenuExtenderDelegates.IsEmpty())
+	{
+		for (const FDataLayerEditorModule::FDataLayersMenuExtender& MenuExtender : MenuExtenderDelegates)
+		{
+			if (MenuExtender.IsBound())
+			{
+				Extenders.Add(MenuExtender.Execute(Commands.ToSharedRef(), SelectedDataLayersSet, SelectedDataLayerActors));
+			}
+		}
+		if (!Extenders.IsEmpty())
+		{
+			Context.AddExtender(FExtender::Combine(Extenders));
+		}
+	}
+
 	// Build up the menu for a selection
 	UToolMenus* ToolMenus = UToolMenus::Get();
 	UToolMenu* Menu = ToolMenus->GenerateMenu(MenuName, Context);
@@ -1778,12 +1799,19 @@ TUniquePtr<ISceneOutlinerHierarchy> FDataLayerMode::CreateHierarchy()
 	return DataLayerHierarchy;
 }
 
-void FDataLayerMode::OnItemSelectionChanged(FSceneOutlinerTreeItemPtr TreeItem, ESelectInfo::Type SelectionType, const FSceneOutlinerItemSelection& Selection)
+
+void FDataLayerMode::CacheSelectedItems(const FSceneOutlinerItemSelection& Selection)
 {
 	SelectedDataLayersSet.Empty();
 	SelectedDataLayerActors.Empty();
 	Selection.ForEachItem<FDataLayerTreeItem>([this](const FDataLayerTreeItem& Item) { SelectedDataLayersSet.Add(Item.GetDataLayer()); });
 	Selection.ForEachItem<FDataLayerActorTreeItem>([this](const FDataLayerActorTreeItem& Item) { SelectedDataLayerActors.Add(FSelectedDataLayerActor(Item.GetDataLayer(), Item.GetActor())); });
+}
+
+void FDataLayerMode::OnItemSelectionChanged(FSceneOutlinerTreeItemPtr TreeItem, ESelectInfo::Type SelectionType, const FSceneOutlinerItemSelection& Selection)
+{
+	CacheSelectedItems(Selection);
+
 	if (DataLayerBrowser)
 	{
 		DataLayerBrowser->OnSelectionChanged(SelectedDataLayersSet);
diff --git a/Engine/Source/Editor/DataLayerEditor/Private/DataLayer/DataLayerMode.h b/Engine/Source/Editor/DataLayerEditor/Private/DataLayer/DataLayerMode.h
--- a/Engine/Source/Editor/DataLayerEditor/Private/DataLayer/DataLayerMode.h
+++ b/Engine/Source/Editor/DataLayerEditor/Private/DataLayer/DataLayerMode.h
@@ -149,6 +149,7 @@ private:
 	bool ContainsSelectedChildDataLayer(const UDataLayerInstance* DataLayer) const;
 	void DeleteDataLayers(const TArray<UDataLayerInstance*>& InDataLayersToDelete);
 	void RefreshSelection();
+	void CacheSelectedItems(const FSceneOutlinerItemSelection& Selection);
 	UWorld* GetOwningWorld() const;
 	AWorldDataLayers* GetOwningWorldAWorldDataLayers() const;
 	FSceneOutlinerDragValidationInfo ValidateActorDrop(const ISceneOutlinerTreeItem& DropTarget, bool bMoveOperation = false) const;
diff --git a/Engine/Source/Editor/DataLayerEditor/Public/DataLayerEditorModule.h b/Engine/Source/Editor/DataLayerEditor/Public/DataLayerEditorModule.h
--- a/Engine/Source/Editor/DataLayerEditor/Public/DataLayerEditorModule.h
+++ b/Engine/Source/Editor/DataLayerEditor/Public/DataLayerEditorModule.h
@@ -44,8 +44,9 @@ public:
 	virtual void SyncDataLayerBrowserToDataLayer(const UDataLayerInstance* DataLayer);
 
 	/** Delegates to be called to extend the DataLayers menus */
-	DECLARE_DELEGATE_RetVal_OneParam(TSharedRef<FExtender>, FDataLayersMenuExtender, const TSharedRef<FUICommandList>);
-	virtual TArray<FDataLayersMenuExtender>& GetAllDataLayersMenuExtenders() {return DataLayersMenuExtenders;}
+	typedef TPair<TWeakObjectPtr<const UDataLayerInstance>, TWeakObjectPtr<const AActor>> FDataLayerActor;
+	DECLARE_DELEGATE_RetVal_ThreeParams(TSharedRef<FExtender>, FDataLayersMenuExtender, const TSharedRef<FUICommandList>, const TSet<TWeakObjectPtr<const UDataLayerInstance>>& /*SelectedDataLayers*/,  const TSet<FDataLayerActor>& /*SelectedDataLayerActors*/);
+	virtual TArray<FDataLayersMenuExtender>& GetAllDataLayersMenuExtenders() { return DataLayersMenuExtenders; }
 
 	/* Implement IDataLayerEditorModule */
 	virtual bool AddActorToDataLayers(AActor* Actor, const TArray<UDataLayerInstance*>& DataLayers) override;
