Fixed crash when opening data layer menu after closing/reopening the data layer outliner

diff --git a/Engine/Source/Editor/DataLayerEditor/Private/DataLayer/DataLayerMode.cpp b/Engine/Source/Editor/DataLayerEditor/Private/DataLayer/DataLayerMode.cpp
--- a/Engine/Source/Editor/DataLayerEditor/Private/DataLayer/DataLayerMode.cpp
+++ b/Engine/Source/Editor/DataLayerEditor/Private/DataLayer/DataLayerMode.cpp
@@ -1229,7 +1229,7 @@ void FDataLayerMode::RegisterContextMenu()
 	{
 		UToolMenu* Menu = ToolMenus->RegisterMenu(DefaultContextBaseMenuName);
 
-		Menu->AddDynamicSection("DataLayerDynamicSection", FNewToolMenuDelegate::CreateLambda([this](UToolMenu* InMenu)
+		Menu->AddDynamicSection("DataLayerDynamicSection", FNewToolMenuDelegate::CreateLambda([](UToolMenu* InMenu)
 		{
 			USceneOutlinerMenuContext* Context = InMenu->FindContext<USceneOutlinerMenuContext>();
 			if (!Context || !Context->SceneOutliner.IsValid())
@@ -1587,7 +1587,7 @@ void FDataLayerMode::RegisterContextMenu()
 			if (!UDataLayerEditorSubsystem::Get()->HasDeprecatedDataLayers())
 			{
 				FToolMenuSection& Section = InMenu->AddSection("AssetOptionsSection", LOCTEXT("AssetOptionsText", "Asset Options"));
-				Section.AddMenuEntryWithCommandList(FGlobalEditorCommonCommands::Get().FindInContentBrowser, Commands);
+				Section.AddMenuEntryWithCommandList(FGlobalEditorCommonCommands::Get().FindInContentBrowser, Mode->Commands);
 			}
 		}));
 	}
