Allow editing data layers of level instance when editing the level instance
- When editing level instance, only show loaded/unloaded actors of this level
- When editing level instance, don't prefix actor label with parent level instance(s)
- Replaced import data layer menu option by drag and dropping level instance data layers into main world AWorldDataLayers
- Removed menu option "Create Data Layer under <other> DataLayer" as it will never be allowed for DataLayerInstanceWithAsset since created without any valid DataLayerAsset.
- Disabled changing visibility and editor loading of DataLayerInstances of editing level instance
- Allow editing data layer of level instance when editing the level instance even when main world is not partitioned
- Changed font from bold to bold+italic on DataLayerInstances containing actors part of editor selection
- Removed suffix "- WorldDataLayers" on parent node of DataLayerInstances in the DataLayer Outliner
- Colorize and append "(Current)" to WorldDataLayer of editing Level Instance in the DataLayer Outliner
- Support Make Current on Level instance DataLayerInstances when editing level instance

diff --git a/Engine/Source/Editor/DataLayerEditor/Private/DataLayer/DataLayerActorTreeItem.h b/Engine/Source/Editor/DataLayerEditor/Private/DataLayer/DataLayerActorTreeItem.h
--- a/Engine/Source/Editor/DataLayerEditor/Private/DataLayer/DataLayerActorTreeItem.h
+++ b/Engine/Source/Editor/DataLayerEditor/Private/DataLayer/DataLayerActorTreeItem.h
@@ -99,7 +99,7 @@ private:
 		{
 			ULevel* Level = Actor.IsValid() ? Actor->GetLevel() : nullptr;
 			ULevelInstanceSubsystem* LevelInstanceSubsystem = UWorld::GetSubsystem<ULevelInstanceSubsystem>(OwningWorld);
-			if (LevelInstanceSubsystem && Level)
+			if (LevelInstanceSubsystem && Level && (Level != OwningWorld->GetCurrentLevel()))
 			{
 				DisplayString = LevelInstanceSubsystem->PrefixWithParentLevelInstanceActorLabels(DisplayString, Level);
 			}
diff --git a/Engine/Source/Editor/DataLayerEditor/Private/DataLayer/DataLayerDragDropOp.cpp b/Engine/Source/Editor/DataLayerEditor/Private/DataLayer/DataLayerDragDropOp.cpp
--- a/Engine/Source/Editor/DataLayerEditor/Private/DataLayer/DataLayerDragDropOp.cpp
+++ b/Engine/Source/Editor/DataLayerEditor/Private/DataLayer/DataLayerDragDropOp.cpp
@@ -18,13 +18,13 @@ struct FSlateBrush;
 void FDataLayerDragDropOp::Construct()
 {
 	const FSlateBrush* Icon = FAppStyle::GetBrush(TEXT("DataLayer.Editor"));
-	if (DataLayerDragDropInfos.Num() == 1)
+	if (DataLayerInstances.Num() == 1 && DataLayerInstances[0].IsValid())
 	{
-		SetToolTip(FText::FromString(DataLayerDragDropInfos[0].DataLayerShortName), Icon);
+		SetToolTip(FText::FromString(DataLayerInstances[0]->GetDataLayerShortName()), Icon);
 	}
 	else
 	{
-		FText Text = FText::Format(NSLOCTEXT("FDataLayerDragDropOp", "MultipleFormat", "{0} DataLayers"), DataLayerDragDropInfos.Num());
+		FText Text = FText::Format(NSLOCTEXT("FDataLayerDragDropOp", "MultipleFormat", "{0} DataLayers"), DataLayerInstances.Num());
 		SetToolTip(Text, Icon);
 	}
 
diff --git a/Engine/Source/Editor/DataLayerEditor/Private/DataLayer/DataLayerEditorSubsystem.cpp b/Engine/Source/Editor/DataLayerEditor/Private/DataLayer/DataLayerEditorSubsystem.cpp
--- a/Engine/Source/Editor/DataLayerEditor/Private/DataLayer/DataLayerEditorSubsystem.cpp
+++ b/Engine/Source/Editor/DataLayerEditor/Private/DataLayer/DataLayerEditorSubsystem.cpp
@@ -683,41 +683,7 @@ bool UDataLayerEditorSubsystem::SelectActorsInDataLayer(UDataLayerInstance* Data
 
 bool UDataLayerEditorSubsystem::SelectActorsInDataLayer(UDataLayerInstance* DataLayerInstance, const bool bSelect, const bool bNotify, const bool bSelectEvenIfHidden, const TSharedPtr<FActorFilter>& Filter)
 {
-	bool bChangesOccurred = false;
-
-	GEditor->GetSelectedActors()->BeginBatchSelectOperation();
-	{
-		// Iterate over all actors, looking for actors in the specified DataLayers.
-		for (AActor* Actor : FActorRange(GetWorld()))
-		{
-			if (!IsActorValidForDataLayer(Actor))
-			{
-				continue;
-			}
-
-			if (Filter.IsValid() && !Filter->PassesFilter(Actor))
-			{
-				continue;
-			}
-
-			if (Actor->ContainsDataLayer(DataLayerInstance))
-			{
-				// The actor was found to be in a specified DataLayerInstance. Set selection state and move on to the next actor.
-				bool bNotifyForActor = false;
-				GEditor->GetSelectedActors()->Modify();
-				GEditor->SelectActor(Actor, bSelect, bNotifyForActor, bSelectEvenIfHidden);
-				bChangesOccurred = true;
-			}
-		}
-	}
-	GEditor->GetSelectedActors()->EndBatchSelectOperation();
-
-	if (bNotify)
-	{
-		GEditor->NoteSelectionChange();
-	}
-
-	return bChangesOccurred;
+	return SelectActorsInDataLayers({ DataLayerInstance }, bSelect, bNotify, bSelectEvenIfHidden, Filter);
 }
 
 bool UDataLayerEditorSubsystem::SelectActorsInDataLayers(const TArray<UDataLayerInstance*>& DataLayers, const bool bSelect, const bool bNotify, const bool bSelectEvenIfHidden)
@@ -743,14 +709,14 @@ bool UDataLayerEditorSubsystem::SelectActorsInDataLayers(const TArray<UDataLayer
 			continue;
 		}
 
-		if (Filter.IsValid() && !Filter->PassesFilter(TWeakObjectPtr<AActor>(Actor)))
+		if (Filter.IsValid() && !Filter->PassesFilter(Actor))
 		{
 			continue;
 		}
 
 		for (const UDataLayerInstance* DataLayerInstance : DataLayers)
 		{
-			if (Actor->ContainsDataLayer(DataLayerInstance))
+			if (Actor->ContainsDataLayer(DataLayerInstance) || Actor->GetDataLayerInstancesForLevel().Contains(DataLayerInstance))
 			{
 				// The actor was found to be in a specified DataLayerInstance. Set selection state and move on to the next actor.
 				bool bNotifyForActor = false;
@@ -891,25 +857,15 @@ void UDataLayerEditorSubsystem::AppendActorsFromDataLayer(UDataLayerInstance* Da
 
 void UDataLayerEditorSubsystem::AppendActorsFromDataLayer(UDataLayerInstance* DataLayerInstance, TArray<AActor*>& InOutActors, const TSharedPtr<FActorFilter>& Filter) const
 {
-	for (AActor* Actor : FActorRange(GetWorld()))
-	{
-		if (Filter.IsValid() && !Filter->PassesFilter(Actor))
-		{
-			continue;
-		}
-		if (Actor->ContainsDataLayer(DataLayerInstance))
-		{
-			InOutActors.Add(Actor);
-		}
-	}
+	AppendActorsFromDataLayers({ DataLayerInstance }, InOutActors, Filter);
 }
 
-void UDataLayerEditorSubsystem::AppendActorsFromDataLayers(const TArray<UDataLayerInstance*>& DataLayers, TArray<AActor*>& InOutActors) const
+void UDataLayerEditorSubsystem::AppendActorsFromDataLayers(const TArray<UDataLayerInstance*>& DataLayerInstances, TArray<AActor*>& InOutActors) const
 {
-	AppendActorsFromDataLayers(DataLayers, InOutActors, nullptr);
+	AppendActorsFromDataLayers(DataLayerInstances, InOutActors, nullptr);
 }
 
-void UDataLayerEditorSubsystem::AppendActorsFromDataLayers(const TArray<UDataLayerInstance*>& DataLayers, TArray<AActor*>& InOutActors, const TSharedPtr<FActorFilter>& Filter) const
+void UDataLayerEditorSubsystem::AppendActorsFromDataLayers(const TArray<UDataLayerInstance*>& DataLayerInstances, TArray<AActor*>& InOutActors, const TSharedPtr<FActorFilter>& Filter) const
 {
 	for (AActor* Actor : FActorRange(GetWorld()))
 	{
@@ -917,12 +873,15 @@ void UDataLayerEditorSubsystem::AppendActorsFromDataLayers(const TArray<UDataLay
 		{
 			continue;
 		}
-		for (const UDataLayerInstance* DataLayerInstance : DataLayers)
+		for (const UDataLayerInstance* DataLayerInstance : DataLayerInstances)
 		{
-			if (Actor->ContainsDataLayer(DataLayerInstance))
+			if (DataLayerInstance)
 			{
-				InOutActors.Add(Actor);
-				break;
+				if (Actor->ContainsDataLayer(DataLayerInstance) || Actor->GetDataLayerInstancesForLevel().Contains(DataLayerInstance))
+				{
+					InOutActors.Add(Actor);
+					break;
+				}
 			}
 		}
 	}
@@ -1318,6 +1277,10 @@ const TSet<TWeakObjectPtr<const UDataLayerInstance>>& UDataLayerEditorSubsystem:
 			{
 				SelectedDataLayersFromEditorSelection.Add(DataLayerInstance);
 			}
+			for (const UDataLayerInstance* DataLayerInstance : Actor->GetDataLayerInstancesForLevel())
+			{
+				SelectedDataLayersFromEditorSelection.Add(DataLayerInstance);
+			}
 		}
 	}
 	return SelectedDataLayersFromEditorSelection;
diff --git a/Engine/Source/Editor/DataLayerEditor/Private/DataLayer/DataLayerHierarchy.cpp b/Engine/Source/Editor/DataLayerEditor/Private/DataLayer/DataLayerHierarchy.cpp
--- a/Engine/Source/Editor/DataLayerEditor/Private/DataLayer/DataLayerHierarchy.cpp
+++ b/Engine/Source/Editor/DataLayerEditor/Private/DataLayer/DataLayerHierarchy.cpp
@@ -24,6 +24,7 @@
 #include "WorldDataLayersTreeItem.h"
 #include "WorldPartition/ActorDescContainerCollection.h"
 #include "WorldPartition/DataLayer/DataLayerInstance.h"
+#include "WorldPartition/DataLayer/DataLayerInstanceWithAsset.h"
 #include "WorldPartition/DataLayer/DataLayerManager.h"
 #include "WorldPartition/DataLayer/WorldDataLayers.h"
 #include "WorldPartition/IWorldPartitionEditorModule.h"
@@ -207,38 +208,13 @@ void FDataLayerHierarchy::CreateItems(TArray<FSceneOutlinerTreeItemPtr>& OutItem
 
 		if (bShowDataLayerActors)
 		{
-			const ULevelInstanceSubsystem* LevelInstanceSubsystem = UWorld::GetSubsystem<ULevelInstanceSubsystem>(GetOwningWorld());
-
-			// Build a set of considered level instances to show nested level instance actors and their child actors
-			TSet<const ILevelInstanceInterface*> ConsideredLevelInstances;		
-			if (bShowLevelInstanceContent && CurrentLevel && LevelInstanceSubsystem)
-			{
-				ILevelInstanceInterface* CurrentLevelInstance = LevelInstanceSubsystem->GetOwningLevelInstance(CurrentLevel);
-				if (LevelInstanceSubsystem->IsEditingLevelInstance(CurrentLevelInstance))
-				{
-					const bool bRecursive = true;
-					LevelInstanceSubsystem->ForEachLevelInstanceChild(CurrentLevelInstance, bRecursive, [&ConsideredLevelInstances](const ILevelInstanceInterface* ChildLevelInstance)
-					{
-						ConsideredLevelInstances.Add(ChildLevelInstance);
-						return true;
-					});
-				}
-			}
-
 			for (AActor* Actor : FActorRange(GetOwningWorld()))
 			{
 				// Consider all actors or actors part of current level (if there is one)
 				bool bConsiderActor = !CurrentLevel || (Actor->GetLevel() == CurrentLevel);
-				if (!bConsiderActor && bShowLevelInstanceContent && LevelInstanceSubsystem)
-				{
-					// If current level is referenced by editing level instance, also consider child actors part of any child Level Instances of this level
-					ILevelInstanceInterface* ParentLevelInstance = LevelInstanceSubsystem->GetParentLevelInstance(Actor);
-					bConsiderActor = ParentLevelInstance && ConsideredLevelInstances.Contains(ParentLevelInstance);
-				}
-
 				if (bConsiderActor)
 				{
-					for (const UDataLayerInstance* DataLayerInstance : Actor->GetDataLayerInstances())
+					for (const UDataLayerInstance* DataLayerInstance : (Actor->GetLevel() == CurrentLevel) ? Actor->GetDataLayerInstancesForLevel() : Actor->GetDataLayerInstances())
 					{
 						if (IsDataLayerShown(DataLayerInstance))
 						{
@@ -253,9 +229,36 @@ void FDataLayerHierarchy::CreateItems(TArray<FSceneOutlinerTreeItemPtr>& OutItem
 
 			if (bShowUnloadedActors)
 			{
-				if (UWorldPartitionSubsystem* WorldPartitionSubsystem = UWorld::GetSubsystem<UWorldPartitionSubsystem>(GetOwningWorld()))
+				// Build mapping of OwninWorld DataLayerInstance to CurrentLevel DataLayerInstance
+				TMap<UDataLayerInstance*, UDataLayerInstance*> WorldToLevelDataLayerMap;
+				if (CurrentLevel)
+				{
+					AWorldDataLayers* OwningWorldWorldDataLayers = GetOwningWorld()->GetWorldDataLayers();
+					AWorldDataLayers* CurrentLevelWorldDataLayers = CurrentLevel->GetWorldDataLayers();
+					if (OwningWorldWorldDataLayers && CurrentLevelWorldDataLayers)
+					{
+						check(OwningWorldWorldDataLayers != CurrentLevelWorldDataLayers);
+						OwningWorldWorldDataLayers->ForEachDataLayer([&WorldToLevelDataLayerMap, CurrentLevelWorldDataLayers](UDataLayerInstance* DataLayerInstance)
+						{
+							if (UDataLayerInstanceWithAsset* DataLayerInstanceWithAsset = Cast<UDataLayerInstanceWithAsset>(DataLayerInstance))
+							{
+								const UDataLayerAsset* DataLayerAsset = DataLayerInstanceWithAsset->GetAsset();
+								if (const UDataLayerInstance* LevelDataLayerInstance = DataLayerAsset ? CurrentLevelWorldDataLayers->GetDataLayerInstance(DataLayerAsset) : nullptr)
+								{
+									WorldToLevelDataLayerMap.Add(DataLayerInstance, const_cast<UDataLayerInstance*>(LevelDataLayerInstance));
+								}
+							}
+							return true;
+						});
+					}
+				}
+
+				UWorldPartitionSubsystem* WorldPartitionSubsystem = UWorld::GetSubsystem<UWorldPartitionSubsystem>(GetOwningWorld());
+				const UDataLayerManager* DataLayerManager = GetOwningWorld()->GetDataLayerManager();
+				if (WorldPartitionSubsystem && DataLayerManager)
 				{
-					WorldPartitionSubsystem->ForEachWorldPartition([this, CurrentLevel, LevelInstanceSubsystem, IsDataLayerShown, &OutItems](UWorldPartition* WorldPartition)
+					const ULevelInstanceSubsystem* LevelInstanceSubsystem = UWorld::GetSubsystem<ULevelInstanceSubsystem>(GetOwningWorld());
+					WorldPartitionSubsystem->ForEachWorldPartition([this, DataLayerManager, CurrentLevel, LevelInstanceSubsystem, IsDataLayerShown, &WorldToLevelDataLayerMap, &OutItems](UWorldPartition* WorldPartition)
 					{
 						// Skip WorldPartition if it's not the one of the current level (the editing level instance)
 						// or if we hide the content of level instances
@@ -268,18 +271,24 @@ void FDataLayerHierarchy::CreateItems(TArray<FSceneOutlinerTreeItemPtr>& OutItem
 						}
 
 						// Create an FDataLayerActorDescTreeItem for each unloaded actor of this WorldPartition
-						const UDataLayerManager* DataLayerManager = WorldPartition->GetDataLayerManager();
-						FWorldPartitionHelpers::ForEachActorDesc(WorldPartition, [this, IsDataLayerShown, DataLayerManager, &OutItems](const FWorldPartitionActorDesc* ActorDesc)
+						FWorldPartitionHelpers::ForEachActorDesc(WorldPartition, [this, IsDataLayerShown, DataLayerManager, CurrentLevel, &WorldToLevelDataLayerMap, &OutItems](const FWorldPartitionActorDesc* ActorDesc)
 						{
 							if (ActorDesc != nullptr && !ActorDesc->IsLoaded())
 							{
 								for (const FName& DataLayerInstanceName : ActorDesc->GetDataLayerInstanceNames())
 								{
-									if (const UDataLayerInstance* const DataLayerInstance = DataLayerManager->GetDataLayerInstance(DataLayerInstanceName))
+									if (const UDataLayerInstance* DataLayerInstance = DataLayerManager->GetDataLayerInstance(DataLayerInstanceName))
 									{
 										if (IsDataLayerShown(DataLayerInstance))
 										{
-											if (const FSceneOutlinerTreeItemPtr ActorDescItem = Mode->CreateItemFor<FDataLayerActorDescTreeItem>(FDataLayerActorDescTreeItemData(ActorDesc->GetGuid(), ActorDesc->GetContainer(), const_cast<UDataLayerInstance*>(DataLayerInstance))))
+											UDataLayerInstance* EffectiveDataLayerInstance = const_cast<UDataLayerInstance*>(DataLayerInstance);
+											// If current level is valid, build Item with current level's DataLayerInstance
+											if (CurrentLevel)
+											{
+												UDataLayerInstance** LevelDataLayerInstance = WorldToLevelDataLayerMap.Find(DataLayerInstance);
+												EffectiveDataLayerInstance = *LevelDataLayerInstance;
+											}
+											if (const FSceneOutlinerTreeItemPtr ActorDescItem = Mode->CreateItemFor<FDataLayerActorDescTreeItem>(FDataLayerActorDescTreeItemData(ActorDesc->GetGuid(), ActorDesc->GetContainer(), EffectiveDataLayerInstance)))
 											{
 												OutItems.Add(ActorDescItem);
 											}
diff --git a/Engine/Source/Editor/DataLayerEditor/Private/DataLayer/DataLayerMode.cpp b/Engine/Source/Editor/DataLayerEditor/Private/DataLayer/DataLayerMode.cpp
--- a/Engine/Source/Editor/DataLayerEditor/Private/DataLayer/DataLayerMode.cpp
+++ b/Engine/Source/Editor/DataLayerEditor/Private/DataLayer/DataLayerMode.cpp
@@ -614,20 +614,66 @@ FSceneOutlinerDragValidationInfo FDataLayerMode::ValidateDrop(const ISceneOutlin
 		return ValidateActorDrop(DropTarget);
 	}
 	
-	if (!GetDataLayerActorPairsFromOperation(Payload.SourceOperation).IsEmpty()) // Moving actor(s) into a Data Layer
+	// Drag and drop actor(s) into a Data Layer
+	if (!GetDataLayerActorPairsFromOperation(Payload.SourceOperation).IsEmpty())
 	{
 		return ValidateActorDrop(DropTarget, true);
 	}
-	 
-	TArray<const UDataLayerAsset*> DataLayerAssets = GetDataLayerAssetsFromOperation(Payload.SourceOperation);
-	if (!DataLayerAssets.IsEmpty())
+	
+	// Drag and drop DataLayer asset(s)
 	{
-		return ValidateDataLayerAssetDrop(DropTarget, DataLayerAssets);
+		TArray<const UDataLayerAsset*> DataLayerAssets = GetDataLayerAssetsFromOperation(Payload.SourceOperation);
+		if (!DataLayerAssets.IsEmpty())
+		{
+			return ValidateDataLayerAssetDrop(DropTarget, DataLayerAssets);
+		}
 	}
 	
+	// Drag and drop DataLayerInstance(s)
 	TArray<UDataLayerInstance*> PayloadDataLayers = GetDataLayerInstancesFromOperation(Payload.SourceOperation);
 	if (!PayloadDataLayers.IsEmpty())
 	{
+		// Drag and drop DataLayerInstance(s) from different AWorldDataLayers
+		AWorldDataLayers* PayloadWorldDatalayers = nullptr;
+		for (UDataLayerInstance* DataLayerInstance : PayloadDataLayers)
+		{
+			if (!PayloadWorldDatalayers)
+			{
+				PayloadWorldDatalayers = DataLayerInstance->GetOuterWorldDataLayers();
+			}
+			else if (PayloadWorldDatalayers != DataLayerInstance->GetOuterWorldDataLayers())
+			{
+				PayloadWorldDatalayers = nullptr;
+				break;
+			}
+		}
+		if (!PayloadWorldDatalayers)
+		{
+			return FSceneOutlinerDragValidationInfo(ESceneOutlinerDropCompatibility::IncompatibleGeneric, LOCTEXT("CantMoveMultipleDataLayersComingFromDifferentWorldDataLayers", "Can't move multiple Data Layers coming from different WorldDataLayers"));
+		}
+
+		// Drag and drop DataLayerInstanceWithAsset(s) between AWorldDataLayers
+		AWorldDataLayers* DropTargetWorldDataLayers = GetWorldDataLayersFromTreeItem(DropTarget, true);
+		if (DropTargetWorldDataLayers && (PayloadWorldDatalayers != DropTargetWorldDataLayers))
+		{
+			TSet<const UDataLayerAsset*> DataLayerAssets;
+			for (UDataLayerInstance* DataLayerInstance : PayloadDataLayers)
+			{
+				if (const UDataLayerInstanceWithAsset* DataLayerInstanceWithAsset = Cast<UDataLayerInstanceWithAsset>(DataLayerInstance))
+				{
+					if (const UDataLayerAsset* DataLayerAsset = DataLayerInstanceWithAsset->GetAsset())
+					{
+						DataLayerAssets.Add(DataLayerAsset);
+					}
+				}
+			}
+			if (!DataLayerAssets.IsEmpty())
+			{
+				return ValidateDataLayerAssetDrop(DropTarget, DataLayerAssets.Array());
+			}
+			return FSceneOutlinerDragValidationInfo(ESceneOutlinerDropCompatibility::IncompatibleGeneric, LOCTEXT("NoValidDataLayersToMove", "No valid Data Layer to move"));
+		}
+
 		const FDataLayerTreeItem* DataLayerItem = DropTarget.CastTo<FDataLayerTreeItem>();
 		const FDataLayerActorTreeItem* DataLayerActorTreeItem = DropTarget.CastTo<FDataLayerActorTreeItem>();
 		const UDataLayerInstance* ParentDataLayer = DataLayerItem ? DataLayerItem->GetDataLayer() : nullptr;
@@ -684,6 +730,23 @@ FSceneOutlinerDragValidationInfo FDataLayerMode::ValidateDrop(const ISceneOutlin
 	return FSceneOutlinerDragValidationInfo::Invalid();
 }
 
+AWorldDataLayers* FDataLayerMode::GetWorldDataLayersFromTreeItem(const ISceneOutlinerTreeItem& TreeItem, bool bDefaultToOwningWorldIfNotFound) const
+{
+	AWorldDataLayers* WorldDataLayers = bDefaultToOwningWorldIfNotFound ? GetOwningWorldAWorldDataLayers() : nullptr;
+	if (const FWorldDataLayersTreeItem* WorldDataLayerTreeItem = TreeItem.CastTo<FWorldDataLayersTreeItem>())
+	{
+		WorldDataLayers = WorldDataLayerTreeItem->GetWorldDataLayers();
+	}
+	else if (const FDataLayerTreeItem* DataLayerTargetItem = TreeItem.CastTo<FDataLayerTreeItem>())
+	{
+		if (const UDataLayerInstance* DataLayerTarget = DataLayerTargetItem->GetDataLayer())
+		{
+			WorldDataLayers = DataLayerTarget->GetOuterWorldDataLayers();
+		}
+	}
+	return WorldDataLayers;
+}
+
 FSceneOutlinerDragValidationInfo FDataLayerMode::ValidateDataLayerAssetDrop(const ISceneOutlinerTreeItem& DropTarget, const TArray<const UDataLayerAsset*>& DataLayerAssetsToDrop) const
 {
 	check(!DataLayerAssetsToDrop.IsEmpty());
@@ -703,10 +766,11 @@ FSceneOutlinerDragValidationInfo FDataLayerMode::ValidateDataLayerAssetDrop(cons
 		return FSceneOutlinerDragValidationInfo(ESceneOutlinerDropCompatibility::Incompatible, FText::Format(LOCTEXT("DropTargetIsReadOnly", "\"{0}\" is read only."), FText::FromString(DropTarget.GetDisplayString())));
 	}
 
-	TArray<const UDataLayerInstance*> ExistingDataLayerInstance;
-	ExistingDataLayerInstance = DataLayerEditorSubsystem->GetDataLayerInstances(DataLayerAssetsToDrop);
+	AWorldDataLayers* DropTargetWorldDataLayers = GetWorldDataLayersFromTreeItem(DropTarget, true);
+	
+	TArray<const UDataLayerInstance*> ExistingDataLayerInstances = DropTargetWorldDataLayers->GetDataLayerInstances(DataLayerAssetsToDrop);
 
-	if (ExistingDataLayerInstance.IsEmpty())
+	if (ExistingDataLayerInstances.IsEmpty())
 	{
 		if (DropTarget.CastTo<FDataLayerTreeItem>() || DropTarget.CastTo<FWorldDataLayersTreeItem>())
 		{
@@ -720,7 +784,7 @@ FSceneOutlinerDragValidationInfo FDataLayerMode::ValidateDataLayerAssetDrop(cons
 	else
 	{
 		TArray<FString> ExistingAssestNames;
-		Algo::Transform(ExistingDataLayerInstance, ExistingAssestNames, [](const UDataLayerInstance* DataLayerInstance) { return CastChecked<UDataLayerInstanceWithAsset>(DataLayerInstance)->GetAsset()->GetName(); });
+		Algo::Transform(ExistingDataLayerInstances, ExistingAssestNames, [](const UDataLayerInstance* DataLayerInstance) { return CastChecked<UDataLayerInstanceWithAsset>(DataLayerInstance)->GetAsset()->GetName(); });
 		
 		FString ExistingAssetNamesString = FString::Join(ExistingAssestNames, TEXT(","));
 
@@ -735,10 +799,9 @@ TArray<UDataLayerInstance*> FDataLayerMode::GetDataLayerInstancesFromOperation(c
 	
 	auto GetDataLayers = [this, &OutDataLayerInstances](const FDataLayerDragDropOp& DataLayerOp)
 	{
-		const TArray<FDataLayerDragDropOp::FDragDropInfo>& DragDropInfos = DataLayerOp.DataLayerDragDropInfos;
-		for (const FDataLayerDragDropOp::FDragDropInfo& DragDropInfo : DragDropInfos)
+		for (const TWeakObjectPtr<UDataLayerInstance>& DragDropInfo : DataLayerOp.DataLayerInstances)
 		{
-			if (UDataLayerInstance* DataLayerInstance = DataLayerEditorSubsystem->GetDataLayerInstance(DragDropInfo.DataLayerInstanceName))
+			if (UDataLayerInstance* DataLayerInstance = DragDropInfo.Get())
 			{
 				OutDataLayerInstances.AddUnique(DataLayerInstance);
 			}
@@ -914,18 +977,39 @@ void FDataLayerMode::OnDrop(ISceneOutlinerTreeItem& DropTarget, const FSceneOutl
 		return;
 	}
 
+	// Handle dropping Data Layer Asset(s)
 	TArray<const UDataLayerAsset*> DataLayerAssetToDrop = GetDataLayerAssetsFromOperation(Payload.SourceOperation);
-	if (!DataLayerAssetToDrop.IsEmpty()) // Dropping Data Layer Assets 
+	if (!DataLayerAssetToDrop.IsEmpty())
 	{
 		OnDataLayerAssetDropped(DataLayerAssetToDrop, DropTarget);
 		return;
 	}
-	
-	// Moving a data layer(s)
+
+	// Handle dropping DataLayerInstance(s)
+	TArray<UDataLayerInstance*> DataLayerInstances = GetDataLayerInstancesFromOperation(Payload.SourceOperation);
+
+	// Handle dropping DataLayerInstanceWithAsset(s) in a different AWorldDataLayers
+	AWorldDataLayers* DropTargetWorldDataLayers = GetWorldDataLayersFromTreeItem(DropTarget, true);
+	AWorldDataLayers* PayloadWorldDatalayers = DataLayerInstances.Num() ? DataLayerInstances[0]->GetOuterWorldDataLayers() : nullptr; // Validation guarantees that they are all part of the same AWorldDataLayers
+	if (DropTargetWorldDataLayers && (PayloadWorldDatalayers != DropTargetWorldDataLayers))
 	{
-		TArray<UDataLayerInstance*> DataLayerInstances = GetDataLayerInstancesFromOperation(Payload.SourceOperation);
-		SetParentDataLayer(DataLayerInstances, TargetDataLayer);
+		TSet<const UDataLayerAsset*> DataLayerAssets;
+		for (UDataLayerInstance* DataLayerInstance : DataLayerInstances)
+		{
+			if (const UDataLayerInstanceWithAsset* DataLayerInstanceWithAsset = Cast<UDataLayerInstanceWithAsset>(DataLayerInstance))
+			{
+				DataLayerAssets.Add(DataLayerInstanceWithAsset->GetAsset());
+			}
+		}
+		if (!DataLayerAssets.IsEmpty())
+		{
+			OnDataLayerAssetDropped(DataLayerAssets.Array(), DropTarget);
+			return;
+		}
 	}
+	
+	// Moving DataLayerInstance(s)
+	SetParentDataLayer(DataLayerInstances, TargetDataLayer);
 }
 
 void FDataLayerMode::OnDataLayerAssetDropped(const TArray<const UDataLayerAsset*>& DroppedDataLayerAsset, ISceneOutlinerTreeItem& DropTarget) const
@@ -936,12 +1020,7 @@ void FDataLayerMode::OnDataLayerAssetDropped(const TArray<const UDataLayerAsset*
 	{
 		FDataLayerCreationParameters Params;
 		Params.DataLayerAsset = const_cast<UDataLayerAsset*>(DataLayerAsset);
-
-		if (FWorldDataLayersTreeItem* WorldDataLayerTreeItem = DropTarget.CastTo<FWorldDataLayersTreeItem>())
-		{
-			Params.WorldDataLayers = WorldDataLayerTreeItem->GetWorldDataLayers();
-		}
-
+		Params.WorldDataLayers = GetWorldDataLayersFromTreeItem(DropTarget, true);
 		UDataLayerInstance* DataLayerInstance = DataLayerEditorSubsystem->CreateDataLayerInstance(Params);
 
 		if (FDataLayerTreeItem* DataLayerItem = DropTarget.CastTo<FDataLayerTreeItem>())
@@ -1075,13 +1154,7 @@ TSharedPtr<FDragDropOperation> FDataLayerMode::CreateDragDropOperation(const FPo
 	{
 		TSharedPtr<FDataLayerDragDropOp> DataLayerOperation = MakeShareable(new FDataLayerDragDropOp);
 		TArray<TWeakObjectPtr<UDataLayerInstance>> DataLayers = DraggedObjects.GetData<TWeakObjectPtr<UDataLayerInstance>>(FWeakDataLayerSelector());
-		for (const auto& DataLayerInstance : DataLayers)
-		{
-			if (DataLayerInstance.IsValid())
-			{
-				DataLayerOperation->DataLayerDragDropInfos.Emplace(DataLayerInstance.Get());
-			}
-		}
+		DataLayerOperation->Init(DataLayers);
 		DataLayerOperation->Construct();
 		return DataLayerOperation;
 	};
@@ -1149,16 +1222,6 @@ TArray<UDataLayerInstance*> FDataLayerMode::GetSelectedDataLayers(SSceneOutliner
 	return MoveTemp(ValidSelectedDataLayers);
 }
 
-TArray<AWorldDataLayers*> FDataLayerMode::GetSelectedWorldDataLayers(SSceneOutliner* InSceneOutliner) const
-{
-	FSceneOutlinerItemSelection ItemSelection(InSceneOutliner->GetSelection());
-	TArray<FWorldDataLayersTreeItem*> SelectedWorldDataLayersItems;
-	ItemSelection.Get<FWorldDataLayersTreeItem>(SelectedWorldDataLayersItems);
-	TArray<AWorldDataLayers*> ValidSelectedDataLayers;
-	Algo::TransformIf(SelectedWorldDataLayersItems, ValidSelectedDataLayers, [](const auto Item) { return Item && Item->GetWorldDataLayers(); }, [](const auto Item) { return Item->GetWorldDataLayers(); });
-	return MoveTemp(ValidSelectedDataLayers);
-}
-
 void FDataLayerMode::CreateDataLayerPicker(UToolMenu* InMenu, FOnDataLayerPicked OnDataLayerPicked, bool bInShowRoot /*= false*/)
 {
 	if (bInShowRoot)
@@ -1259,17 +1322,43 @@ void FDataLayerMode::RegisterContextMenu()
 			{
 				FToolMenuSection& Section = InMenu->AddSection("DataLayers", LOCTEXT("DataLayers", "Data Layers"));
 				
-				auto CreateNewDataLayerInternal = [SceneOutliner, Mode](UDataLayerInstance* InParentDataLayer, const UDataLayerAsset* InDataLayerAsset = nullptr) -> UDataLayerInstance*
+				auto GetBestCandidateWorldDataLayersFromSelection = [](SSceneOutliner* SceneOutliner, FDataLayerMode* Mode)
+				{
+					FSceneOutlinerItemSelection Selection(SceneOutliner->GetSelection());
+					TArray<FWorldDataLayersTreeItem*> SelectedWorldDataLayersItems;
+					Selection.Get<FWorldDataLayersTreeItem>(SelectedWorldDataLayersItems);
+					TArray<AWorldDataLayers*> SelectedDataLayers;
+					Algo::TransformIf(SelectedWorldDataLayersItems, SelectedDataLayers, [](const FWorldDataLayersTreeItem* Item) { return Item && Item->GetWorldDataLayers(); }, [](const FWorldDataLayersTreeItem* Item) { return Item->GetWorldDataLayers(); });
+					AWorldDataLayers* BestCandidate = (SelectedDataLayers.Num() == 1) ? SelectedDataLayers[0] : nullptr;
+					if (!BestCandidate)
+					{
+						TSet<AWorldDataLayers*> OuterWorldDataLayers;
+						Algo::TransformIf(Mode->SelectedDataLayersSet, OuterWorldDataLayers, [](const TWeakObjectPtr<const UDataLayerInstance>& DataLayerInstance) { return DataLayerInstance.IsValid() && DataLayerInstance->GetOuterWorldDataLayers(); }, [](const TWeakObjectPtr<const UDataLayerInstance>& DataLayerInstance) { return DataLayerInstance->GetOuterWorldDataLayers(); });
+						BestCandidate = (OuterWorldDataLayers.Num() == 1) ? *OuterWorldDataLayers.CreateConstIterator() : nullptr;
+					}
+					if (!BestCandidate)
+					{
+						TSet<AWorldDataLayers*> OuterWorldDataLayers;
+						Algo::TransformIf(Mode->SelectedDataLayerActors, OuterWorldDataLayers, [](const FSelectedDataLayerActor& DataLayerActor) { return DataLayerActor.Key.IsValid() && DataLayerActor.Key->GetOuterWorldDataLayers(); }, [](const FSelectedDataLayerActor& DataLayerActor) { return DataLayerActor.Key->GetOuterWorldDataLayers(); });
+						BestCandidate = (OuterWorldDataLayers.Num() == 1) ? *OuterWorldDataLayers.CreateConstIterator() : nullptr;
+					}
+					if (!BestCandidate)
+					{
+						BestCandidate = Mode->GetOwningWorldAWorldDataLayers();
+					}
+					return BestCandidate;
+				};
+
+				auto CreateNewDataLayerInternal = [&GetBestCandidateWorldDataLayersFromSelection, SceneOutliner, Mode](UDataLayerInstance* InParentDataLayer, const UDataLayerAsset* InDataLayerAsset = nullptr) -> UDataLayerInstance*
 				{
 					const FScopedTransaction Transaction(LOCTEXT("CreateNewDataLayer", "Create New Data Layer"));
+					AWorldDataLayers* TargetWorldDataLayers = GetBestCandidateWorldDataLayersFromSelection(SceneOutliner, Mode);
 					Mode->SelectedDataLayersSet.Empty();
 					Mode->SelectedDataLayerActors.Empty();
 
-					TArray<AWorldDataLayers*> SelectedWorldDataLayers = Mode->GetSelectedWorldDataLayers(SceneOutliner);
-
 					FDataLayerCreationParameters CreationParams;
 					CreationParams.DataLayerAsset = const_cast<UDataLayerAsset*>(InDataLayerAsset);
-					CreationParams.WorldDataLayers = InParentDataLayer ? InParentDataLayer->GetOuterWorldDataLayers() : (SelectedWorldDataLayers.Num() == 1 ? SelectedWorldDataLayers[0] : Mode->GetOwningWorldAWorldDataLayers());
+					CreationParams.WorldDataLayers = InParentDataLayer ? InParentDataLayer->GetOuterWorldDataLayers() : TargetWorldDataLayers;
 					if (UDataLayerInstance* NewDataLayerInstance = UDataLayerEditorSubsystem::Get()->CreateDataLayerInstance(CreationParams))
 					{
 						UDataLayerEditorSubsystem::Get()->SetParentDataLayer(NewDataLayerInstance, InParentDataLayer);
@@ -1301,66 +1390,10 @@ void FDataLayerMode::RegisterContextMenu()
 				{
 					return CreateNewDataLayerInternal(InParentDataLayer);
 				};
-
-				// If selection contains readonly DataLayerInstances
-				if (Algo::AnyOf(Mode->SelectedDataLayersSet, [](const TWeakObjectPtr<const UDataLayerInstance>& DataLayerInstance) { return DataLayerInstance.IsValid() && DataLayerInstance->IsReadOnly(); }))
-				{
-					// Create a menu entry that allows to creating DataLayerInstance in the world using selection referenced DataLayer Assets
-					TSet<const UDataLayerAsset*> ExistingDataLayerAssets;
-					if (AWorldDataLayers* WorldDataLayer = Mode->GetOwningWorldAWorldDataLayers())
-					{
-						WorldDataLayer->ForEachDataLayer([&ExistingDataLayerAssets](UDataLayerInstance* DataLayerInstance)
-						{
-							UDataLayerInstanceWithAsset* DataLayerWithAsset = Cast<UDataLayerInstanceWithAsset>(DataLayerInstance);
-							if (const UDataLayerAsset* DataLayerAsset = DataLayerWithAsset ? DataLayerWithAsset->GetAsset() : nullptr)
-							{
-								ExistingDataLayerAssets.Add(DataLayerAsset);
-							}
-							return true;
-						});
-					}
-
-					TSet<const UDataLayerAsset*> DataLayerAssetsToImport;
-					for (const TWeakObjectPtr<const UDataLayerInstance>& DataLayerInstance : Mode->SelectedDataLayersSet)
-					{
-						if (DataLayerInstance.IsValid() && DataLayerInstance->IsReadOnly())
-						{
-							const UDataLayerInstanceWithAsset* DataLayerWithAsset = Cast<const UDataLayerInstanceWithAsset>(DataLayerInstance);
-							const UDataLayerAsset* DataLayerAsset = DataLayerWithAsset ? DataLayerWithAsset->GetAsset() : nullptr;
-							if (DataLayerAsset && !ExistingDataLayerAssets.Contains(DataLayerAsset))
-							{
-								DataLayerAssetsToImport.Add(DataLayerAsset);
-							}
-						}
-					}
-
-					if (!DataLayerAssetsToImport.IsEmpty() && !UDataLayerEditorSubsystem::Get()->HasDeprecatedDataLayers())
-					{
-						Section.AddMenuEntry("ImportDataLayers", LOCTEXT("ImportDataLayers", "Import Data Layer(s)"), FText(), FSlateIcon(),
-							FUIAction(FExecuteAction::CreateLambda([DataLayerAssetsToImport, CreateNewDataLayer]()
-								{
-									const FScopedTransaction Transaction(LOCTEXT("ImportDataLayersTransaction", "Import Data Layers"));
-									for (const UDataLayerAsset* DataLayerAsset : DataLayerAssetsToImport)
-									{
-										CreateNewDataLayer(nullptr, DataLayerAsset);
-									}
-								})));
-					}
-
-					// Readonly selected Data Layer Instances will only show the option to import Data Layers
-					return;
-				}
 				
 				Section.AddMenuEntry("CreateNewDataLayer", LOCTEXT("CreateNewDataLayer", "Create New Data Layer"), FText(), FSlateIcon(),
 					FUIAction(FExecuteAction::CreateLambda([CreateNewEmptyDataLayer]() { CreateNewEmptyDataLayer(); })));
 
-				UDataLayerInstance* ParentDataLayer = (Mode->SelectedDataLayersSet.Num() == 1) ? const_cast<UDataLayerInstance*>(Mode->SelectedDataLayersSet.CreateIterator()->Get()) : nullptr;
-				if (ParentDataLayer)
-				{
-					Section.AddMenuEntry("CreateNewDataLayerUnderDataLayer", FText::Format(LOCTEXT("CreateNewDataLayerUnderDataLayer", "Create New Data Layer under \"{0}\""), FText::FromString(ParentDataLayer->GetDataLayerShortName())), FText(), FSlateIcon(),
-					FUIAction(FExecuteAction::CreateLambda([CreateNewEmptyDataLayer, ParentDataLayer]() { CreateNewEmptyDataLayer(ParentDataLayer); })));
-				}
-
 				Section.AddMenuEntry("AddSelectedActorsToSelectedDataLayers", LOCTEXT("AddSelectedActorsToSelectedDataLayersMenu", "Add Selected Actors to Selected Data Layers"), FText(), FSlateIcon(),
 					FUIAction(
 						FExecuteAction::CreateLambda([SelectedDataLayers]() {
diff --git a/Engine/Source/Editor/DataLayerEditor/Private/DataLayer/DataLayerMode.h b/Engine/Source/Editor/DataLayerEditor/Private/DataLayer/DataLayerMode.h
--- a/Engine/Source/Editor/DataLayerEditor/Private/DataLayer/DataLayerMode.h
+++ b/Engine/Source/Editor/DataLayerEditor/Private/DataLayer/DataLayerMode.h
@@ -141,7 +141,6 @@ private:
 	TArray<UDataLayerInstance*> GetDataLayerInstancesFromOperation(const FDragDropOperation& Operation, bool bOnlyFindFirst = false) const;
 	TArray<const UDataLayerAsset*> GetDataLayerAssetsFromOperation(const FDragDropOperation& InDragDrop, bool bOnlyFindFirst = false) const;
 	TArray<UDataLayerInstance*> GetSelectedDataLayers(SSceneOutliner* InSceneOutliner) const;
-	TArray<AWorldDataLayers*> GetSelectedWorldDataLayers(SSceneOutliner* InSceneOutliner) const;
 	void SetParentDataLayer(const TArray<UDataLayerInstance*> DataLayers, UDataLayerInstance* ParentDataLayer) const;
 	void OnLevelSelectionChanged(UObject* Obj);
 	static void CreateDataLayerPicker(UToolMenu* InMenu, FOnDataLayerPicked OnDataLayerPicked, bool bInShowRoot = false);
@@ -152,6 +151,7 @@ private:
 	void CacheSelectedItems(const FSceneOutlinerItemSelection& Selection);
 	UWorld* GetOwningWorld() const;
 	AWorldDataLayers* GetOwningWorldAWorldDataLayers() const;
+	AWorldDataLayers* GetWorldDataLayersFromTreeItem(const ISceneOutlinerTreeItem& TreeItem, bool bDefaultToOwningWorldIfNotFound) const;
 	FSceneOutlinerDragValidationInfo ValidateActorDrop(const ISceneOutlinerTreeItem& DropTarget, bool bMoveOperation = false) const;
 	FSceneOutlinerDragValidationInfo ValidateDataLayerAssetDrop(const ISceneOutlinerTreeItem& DropTarget, const TArray<const UDataLayerAsset*>& DataLayerAssetsToDrop) const;
 	void OnDataLayerAssetDropped(const TArray<const UDataLayerAsset*>& DroppedDataLayerAsset, ISceneOutlinerTreeItem& DropTarget) const;
diff --git a/Engine/Source/Editor/DataLayerEditor/Private/DataLayer/DataLayerNameEditSink.h b/Engine/Source/Editor/DataLayerEditor/Private/DataLayer/DataLayerNameEditSink.h
--- a/Engine/Source/Editor/DataLayerEditor/Private/DataLayer/DataLayerNameEditSink.h
+++ b/Engine/Source/Editor/DataLayerEditor/Private/DataLayer/DataLayerNameEditSink.h
@@ -3,6 +3,7 @@
 #pragma once
 
 #include "DataLayer/DataLayerEditorSubsystem.h"
+#include "WorldPartition/DataLayer/DataLayerInstanceWithAsset.h"
 #include "ScopedTransaction.h"
 #include "IObjectNameEditSink.h"
 
@@ -23,6 +24,11 @@ class FDataLayerNameEditSink : public UE::EditorWidgets::IObjectNameEditSink
 			return FText::Format(FText::FromString("{0}"), FText::FromString(DataLayerInstance->GetDataLayerShortName()));
 		}
 		
+		UDataLayerInstanceWithAsset* DataLayerWithAsset = Cast<UDataLayerInstanceWithAsset>(DataLayerInstance);
+		if (DataLayerWithAsset && !DataLayerWithAsset->GetAsset())
+		{
+			return FText::FromString(TEXT("Unknown"));
+		}
 		return FText::Format(FText::FromString("{0} ({1})"), FText::FromString(DataLayerInstance->GetDataLayerShortName()), FText::FromString(DataLayerInstance->GetDataLayerFullName()));
 	}
 
diff --git a/Engine/Source/Editor/DataLayerEditor/Private/DataLayer/DataLayerOutlinerIsLoadedInEditorColumn.cpp b/Engine/Source/Editor/DataLayerEditor/Private/DataLayer/DataLayerOutlinerIsLoadedInEditorColumn.cpp
--- a/Engine/Source/Editor/DataLayerEditor/Private/DataLayer/DataLayerOutlinerIsLoadedInEditorColumn.cpp
+++ b/Engine/Source/Editor/DataLayerEditor/Private/DataLayer/DataLayerOutlinerIsLoadedInEditorColumn.cpp
@@ -76,7 +76,9 @@ const TSharedRef<SWidget> FDataLayerOutlinerIsLoadedInEditorColumn::ConstructRow
 				{
 					FDataLayerTreeItem* DataLayerTreeItem = TreeItem->CastTo<FDataLayerTreeItem>();
 					const UDataLayerInstance* DataLayerInstance = DataLayerTreeItem->GetDataLayer();
-					return DataLayerInstance && !DataLayerInstance->IsReadOnly() ? EVisibility::Visible : EVisibility::Collapsed;
+					const AWorldDataLayers* OuterWorldDataLayers = DataLayerInstance ? DataLayerInstance->GetOuterWorldDataLayers() : nullptr;
+					const bool bIsSubWorldDataLayers = OuterWorldDataLayers && OuterWorldDataLayers->IsSubWorldDataLayers();
+					return DataLayerInstance && !DataLayerInstance->IsReadOnly() && !bIsSubWorldDataLayers ? EVisibility::Visible : EVisibility::Collapsed;
 				})
 				.IsChecked_Lambda([TreeItem]()
 				{
diff --git a/Engine/Source/Editor/DataLayerEditor/Private/DataLayer/DataLayerPropertyTypeCustomization.cpp b/Engine/Source/Editor/DataLayerEditor/Private/DataLayer/DataLayerPropertyTypeCustomization.cpp
--- a/Engine/Source/Editor/DataLayerEditor/Private/DataLayer/DataLayerPropertyTypeCustomization.cpp
+++ b/Engine/Source/Editor/DataLayerEditor/Private/DataLayer/DataLayerPropertyTypeCustomization.cpp
@@ -243,10 +243,9 @@ FReply FDataLayerPropertyTypeCustomization::OnDrop(const FGeometry& InGeometry,
 	TSharedPtr<const FDataLayerDragDropOp> DataLayerDragDropOp = GetDataLayerDragDropOp(InDragDropEvent.GetOperation());
 	if (DataLayerDragDropOp.IsValid())
 	{
-		const TArray<FDataLayerDragDropOp::FDragDropInfo>& DragDropInfos = DataLayerDragDropOp->DataLayerDragDropInfos;
-		if (ensure(DragDropInfos.Num() == 1))
+		if (ensure(DataLayerDragDropOp->DataLayerInstances.Num() == 1))
 		{
-			if (const UDataLayerInstance* DataLayerPtr = UDataLayerEditorSubsystem::Get()->GetDataLayerInstance(DragDropInfos[0].DataLayerInstanceName))
+			if (const UDataLayerInstance* DataLayerPtr = DataLayerDragDropOp->DataLayerInstances[0].Get())
 			{
 				AssignDataLayer(DataLayerPtr);
 			}
@@ -258,7 +257,7 @@ FReply FDataLayerPropertyTypeCustomization::OnDrop(const FGeometry& InGeometry,
 bool FDataLayerPropertyTypeCustomization::OnVerifyDrag(TSharedPtr<FDragDropOperation> InDragDrop)
 {
 	TSharedPtr<const FDataLayerDragDropOp> DataLayerDragDropOp = GetDataLayerDragDropOp(InDragDrop);
-	return DataLayerDragDropOp.IsValid() && DataLayerDragDropOp->DataLayerDragDropInfos.Num() == 1;
+	return DataLayerDragDropOp.IsValid() && DataLayerDragDropOp->DataLayerInstances.Num() == 1;
 }
 
 TSharedPtr<const FDataLayerDragDropOp> FDataLayerPropertyTypeCustomization::GetDataLayerDragDropOp(TSharedPtr<FDragDropOperation> InDragDrop)
diff --git a/Engine/Source/Editor/DataLayerEditor/Private/DataLayer/DataLayerTreeItem.cpp b/Engine/Source/Editor/DataLayerEditor/Private/DataLayer/DataLayerTreeItem.cpp
--- a/Engine/Source/Editor/DataLayerEditor/Private/DataLayer/DataLayerTreeItem.cpp
+++ b/Engine/Source/Editor/DataLayerEditor/Private/DataLayer/DataLayerTreeItem.cpp
@@ -43,7 +43,9 @@ bool FDataLayerTreeItem::GetVisibility() const
 bool FDataLayerTreeItem::ShouldShowVisibilityState() const
 {
 	const UDataLayerInstance* DataLayerInstancePtr = DataLayerInstance.Get();
-	return DataLayerInstancePtr && !DataLayerInstancePtr->IsReadOnly();
+	const AWorldDataLayers* OuterWorldDataLayers = DataLayerInstancePtr ? DataLayerInstancePtr->GetOuterWorldDataLayers() : nullptr;
+	const bool bIsSubWorldDataLayers = OuterWorldDataLayers && OuterWorldDataLayers->IsSubWorldDataLayers();
+	return DataLayerInstancePtr && !DataLayerInstancePtr->IsReadOnly() && !bIsSubWorldDataLayers;
 }
 
 bool FDataLayerTreeItem::CanInteract() const 
diff --git a/Engine/Source/Editor/DataLayerEditor/Private/DataLayer/DataLayersActorDescTreeItem.h b/Engine/Source/Editor/DataLayerEditor/Private/DataLayer/DataLayersActorDescTreeItem.h
--- a/Engine/Source/Editor/DataLayerEditor/Private/DataLayer/DataLayersActorDescTreeItem.h
+++ b/Engine/Source/Editor/DataLayerEditor/Private/DataLayer/DataLayersActorDescTreeItem.h
@@ -53,7 +53,7 @@ public:
 			UWorld* OwningWorld = Container->GetWorldPartition()->GetWorld();
 			ULevelInstanceSubsystem* LevelInstanceSubsystem = UWorld::GetSubsystem<ULevelInstanceSubsystem>(OwningWorld);
 			ULevel* Level = Container->GetTypedOuter<UWorld>()->PersistentLevel;
-			if (LevelInstanceSubsystem && Level)
+			if (LevelInstanceSubsystem && Level && (Level != OwningWorld->GetCurrentLevel()))
 			{
 				DisplayString = LevelInstanceSubsystem->PrefixWithParentLevelInstanceActorLabels(DisplayString, Level);
 			}
diff --git a/Engine/Source/Editor/DataLayerEditor/Private/DataLayer/SDataLayerBrowser.cpp b/Engine/Source/Editor/DataLayerEditor/Private/DataLayer/SDataLayerBrowser.cpp
--- a/Engine/Source/Editor/DataLayerEditor/Private/DataLayer/SDataLayerBrowser.cpp
+++ b/Engine/Source/Editor/DataLayerEditor/Private/DataLayer/SDataLayerBrowser.cpp
@@ -51,6 +51,7 @@
 #include "WorldPartition/DataLayer/DataLayerInstance.h"
 #include "WorldPartition/WorldPartitionActorDesc.h"
 #include "WorldPartition/WorldPartitionHandle.h"
+#include "Editor.h"
 
 class ISceneOutliner;
 class UObject;
@@ -187,7 +188,18 @@ void SDataLayerBrowser::Construct(const FArguments& InArgs)
 	ChildSlot
 	[
 		SAssignNew(ContentAreaBox, SVerticalBox)
-		.IsEnabled_Lambda([]() { return GWorld ? UWorld::IsPartitionedWorld(GWorld) : false; })
+		.IsEnabled_Lambda([]()
+		{
+			if (UWorld* World = GEditor ? GEditor->GetEditorWorldContext().World() : nullptr)
+			{
+				if (ULevel* CurrentLevel = World->GetCurrentLevel() && World->GetCurrentLevel() != World->PersistentLevel ? World->GetCurrentLevel() : nullptr)
+				{
+					World = CurrentLevel->GetTypedOuter<UWorld>();
+				}
+				return UWorld::IsPartitionedWorld(World);
+			}
+			return false;
+		})
 	];
 
 	InitializeDataLayerBrowser();
diff --git a/Engine/Source/Editor/DataLayerEditor/Private/DataLayer/WorldDataLayersTreeItem.cpp b/Engine/Source/Editor/DataLayerEditor/Private/DataLayer/WorldDataLayersTreeItem.cpp
--- a/Engine/Source/Editor/DataLayerEditor/Private/DataLayer/WorldDataLayersTreeItem.cpp
+++ b/Engine/Source/Editor/DataLayerEditor/Private/DataLayer/WorldDataLayersTreeItem.cpp
@@ -83,10 +83,28 @@ struct SWorldDataLayersTreeLabel : FSceneOutlinerCommonLabelData, public SCompou
 private:
 	TWeakPtr<FWorldDataLayersTreeItem> TreeItemPtr;
 
+	bool IsInActorEditorContext() const
+	{
+		auto Item = TreeItemPtr.Pin();
+		const AWorldDataLayers* WorldDataLayers = Item.IsValid() ? Item->GetWorldDataLayers() : nullptr;
+		const UWorld* OwningWorld = WorldDataLayers ? WorldDataLayers->GetWorld() : nullptr;
+		ULevel* CurrentLevel = (OwningWorld && OwningWorld->GetCurrentLevel() && !OwningWorld->GetCurrentLevel()->IsPersistentLevel()) ? OwningWorld->GetCurrentLevel() : nullptr;
+		return CurrentLevel && (CurrentLevel == WorldDataLayers->GetLevel());
+	}
+
 	FText GetDisplayText() const
 	{
 		auto Item = TreeItemPtr.Pin();
-		return Item.IsValid() ? FText::FromString(Item->GetDisplayString()) : FText();
+		if (Item.IsValid())
+		{
+			FText SuffixText = FText::GetEmpty();
+			if (IsInActorEditorContext())
+			{
+				SuffixText = FText(LOCTEXT("IsCurrentSuffix", " (Current)"));
+			}
+			return FText::Format(LOCTEXT("WorldDataLayersDisplayText", "{0}{1}"), FText::FromString(Item->GetDisplayString()), SuffixText);
+		}
+		return FText();
 	}
 
 	FSlateColor GetForegroundColor() const
@@ -95,6 +113,10 @@ private:
 		{
 			return BaseColor.GetValue();
 		}
+		else if (IsInActorEditorContext())
+		{
+			return FAppStyle::Get().GetSlateColor("Colors.AccentGreen");
+		}
 
 		return FSlateColor::UseForeground();
 	}
@@ -113,14 +135,7 @@ FString FWorldDataLayersTreeItem::GetDisplayString() const
 	if (const AWorldDataLayers* WorldDataLayersPtr = WorldDataLayers.Get())
 	{
 		check(WorldDataLayersPtr->GetLevel());
-		FString DisplayString = SceneOutliner::GetWorldDescription(WorldDataLayersPtr->GetLevel()->GetTypedOuter<UWorld>()).ToString();
-
-		if (!WorldDataLayersPtr->GetActorLabel().IsEmpty())
-		{
-			DisplayString += TEXT(" - ") + WorldDataLayersPtr->GetActorLabel();
-		}
-
-		return DisplayString;
+		return SceneOutliner::GetWorldDescription(WorldDataLayersPtr->GetLevel()->GetTypedOuter<UWorld>()).ToString();
 	}
 	return LOCTEXT("UnknownWorldDataLayers", "Unknown").ToString();
 }
@@ -134,7 +149,7 @@ bool FWorldDataLayersTreeItem::IsReadOnly() const
 {
 	const AWorldDataLayers* WorldDataLayersPtr = WorldDataLayers.Get();
 	UWorld* WorldPtr = WorldDataLayersPtr ? WorldDataLayersPtr->GetLevel()->GetTypedOuter<UWorld>() : nullptr;
-	return (WorldDataLayersPtr && WorldDataLayersPtr->IsSubWorldDataLayers()) || !WorldPtr || (WorldPtr->WorldType != EWorldType::Editor);
+	return (WorldDataLayersPtr && WorldDataLayersPtr->IsReadOnly()) || !WorldPtr || (WorldPtr->WorldType != EWorldType::Editor);
 }
 
 TSharedRef<SWidget> FWorldDataLayersTreeItem::GenerateLabelWidget(ISceneOutliner& Outliner, const STableRow<FSceneOutlinerTreeItemPtr>& InRow)
diff --git a/Engine/Source/Editor/DataLayerEditor/Public/DataLayer/DataLayerDragDropOp.h b/Engine/Source/Editor/DataLayerEditor/Public/DataLayer/DataLayerDragDropOp.h
--- a/Engine/Source/Editor/DataLayerEditor/Public/DataLayer/DataLayerDragDropOp.h
+++ b/Engine/Source/Editor/DataLayerEditor/Public/DataLayer/DataLayerDragDropOp.h
@@ -21,19 +21,19 @@ class FDataLayerDragDropOp : public FDecoratedDragDropOp
 public:
 	DRAG_DROP_OPERATOR_TYPE(FDataLayerDragDropOp, FDecoratedDragDropOp)
 
-	struct FDragDropInfo
+	void Init(const TArray<TWeakObjectPtr<UDataLayerInstance>>& InDataLayerInstances)
 	{
-		FDragDropInfo(const UDataLayerInstance* DataLayerInstance)
-			:DataLayerInstanceName(DataLayerInstance->GetDataLayerFName()),
-			DataLayerShortName(DataLayerInstance->GetDataLayerShortName())
-		{}
-
-		FName DataLayerInstanceName;
-		FString DataLayerShortName;
-	};
+		for (const auto& DataLayerInstance : InDataLayerInstances)
+		{
+			if (DataLayerInstance.IsValid())
+			{
+				DataLayerInstances.Emplace(DataLayerInstance.Get());
+			}
+		}
+	}
 
 	/** The labels of the layers being dragged */
-	TArray<FDragDropInfo> DataLayerDragDropInfos;
+	TArray< TWeakObjectPtr<UDataLayerInstance> > DataLayerInstances;
 
 	virtual void Construct() override;
 };
diff --git a/Engine/Source/Editor/EditorStyle/Private/StarshipStyle.cpp b/Engine/Source/Editor/EditorStyle/Private/StarshipStyle.cpp
--- a/Engine/Source/Editor/EditorStyle/Private/StarshipStyle.cpp
+++ b/Engine/Source/Editor/EditorStyle/Private/StarshipStyle.cpp
@@ -1127,7 +1127,7 @@ void FStarshipEditorStyle::FStyle::SetupGeneralStyles()
 			.SetPressed(BOX_BRUSH("Common/RoundedSelection_16x", 4.0f / 16.0f, SelectionColor_Pressed))
 		);
 		Set("DataLayerBrowser.LabelFont", DEFAULT_FONT("Regular", 9));
-		Set("DataLayerBrowser.LabelFontBold", DEFAULT_FONT("Bold", 10));
+		Set("DataLayerBrowser.LabelFontBold", DEFAULT_FONT("BoldItalic", 10));
 		Set("DataLayer.ColorIcon", new FSlateBoxBrush(NAME_None, 8.0f / 32.0f, FStyleColors::White));
 	}
 
diff --git a/Engine/Source/Editor/MovieSceneTools/Private/TrackEditors/DataLayerTrackEditor.cpp b/Engine/Source/Editor/MovieSceneTools/Private/TrackEditors/DataLayerTrackEditor.cpp
--- a/Engine/Source/Editor/MovieSceneTools/Private/TrackEditors/DataLayerTrackEditor.cpp
+++ b/Engine/Source/Editor/MovieSceneTools/Private/TrackEditors/DataLayerTrackEditor.cpp
@@ -179,7 +179,7 @@ struct FDataLayerSection
 
 	bool OnAllowDrop(TSharedPtr<FDragDropOperation> DragDropOperation)
 	{
-		if (DragDropOperation->IsOfType<FDataLayerDragDropOp>() && StaticCastSharedPtr<FDataLayerDragDropOp>(DragDropOperation)->DataLayerDragDropInfos.Num() > 0)
+		if (DragDropOperation->IsOfType<FDataLayerDragDropOp>() && StaticCastSharedPtr<FDataLayerDragDropOp>(DragDropOperation)->DataLayerInstances.Num() > 0)
 		{
 			return true;
 		}
@@ -207,13 +207,13 @@ struct FDataLayerSection
 		TArray<UDataLayerAsset*> DataLayerAssets = Section->GetDataLayerAssets();
 		if (TSharedPtr<FDataLayerDragDropOp> DataLayerDragDropOp = InDragDropEvent.GetOperationAs<FDataLayerDragDropOp>())
 		{
-			if (DataLayerDragDropOp->DataLayerDragDropInfos.Num() > 0)
+			if (DataLayerDragDropOp->DataLayerInstances.Num() > 0)
 			{
-				for (const FDataLayerDragDropOp::FDragDropInfo& DragDropInfo : DataLayerDragDropOp->DataLayerDragDropInfos)
+				for (const TWeakObjectPtr<UDataLayerInstance>& DataLayerInstance : DataLayerDragDropOp->DataLayerInstances)
 				{
-					if (UDataLayerInstance* DataLayerInstance = SubSystem->GetDataLayerInstance(DragDropInfo.DataLayerInstanceName))
+					if (DataLayerInstance.IsValid())
 					{
-						UDataLayerInstanceWithAsset* DataLayerWithAsset = Cast<UDataLayerInstanceWithAsset>(DataLayerInstance);
+						UDataLayerInstanceWithAsset* DataLayerWithAsset = Cast<UDataLayerInstanceWithAsset>(DataLayerInstance.Get());
 						UDataLayerAsset* DataLayerAsset = DataLayerWithAsset ? const_cast<UDataLayerAsset*>(DataLayerWithAsset->GetAsset()) : nullptr;
 
 						if (DataLayerAsset)
diff --git a/Engine/Source/Runtime/Engine/Private/WorldPartition/DataLayer/DataLayerInstance.cpp b/Engine/Source/Runtime/Engine/Private/WorldPartition/DataLayer/DataLayerInstance.cpp
--- a/Engine/Source/Runtime/Engine/Private/WorldPartition/DataLayer/DataLayerInstance.cpp
+++ b/Engine/Source/Runtime/Engine/Private/WorldPartition/DataLayer/DataLayerInstance.cpp
@@ -151,7 +151,7 @@ bool UDataLayerInstance::AddActor(AActor* InActor) const
 
 bool UDataLayerInstance::CanRemoveActor(AActor* InActor) const
 {
-	return InActor->ContainsDataLayer(this);
+	return InActor->GetDataLayerInstances().Contains(this) || InActor->GetDataLayerInstancesForLevel().Contains(this);
 }
 
 bool UDataLayerInstance::RemoveActor(AActor* InActor) const
diff --git a/Engine/Source/Runtime/Engine/Private/WorldPartition/DataLayer/DataLayerInstanceWithAsset.cpp b/Engine/Source/Runtime/Engine/Private/WorldPartition/DataLayer/DataLayerInstanceWithAsset.cpp
--- a/Engine/Source/Runtime/Engine/Private/WorldPartition/DataLayer/DataLayerInstanceWithAsset.cpp
+++ b/Engine/Source/Runtime/Engine/Private/WorldPartition/DataLayer/DataLayerInstanceWithAsset.cpp
@@ -40,7 +40,7 @@ bool UDataLayerInstanceWithAsset::IsReadOnly() const
 	{
 		return true;
 	}
-	return GetOuterWorldDataLayers()->IsSubWorldDataLayers();
+	return GetOuterWorldDataLayers()->IsReadOnly();
 }
 
 bool UDataLayerInstanceWithAsset::IsLocked() const
diff --git a/Engine/Source/Runtime/Engine/Private/WorldPartition/DataLayer/DataLayerManager.cpp b/Engine/Source/Runtime/Engine/Private/WorldPartition/DataLayer/DataLayerManager.cpp
--- a/Engine/Source/Runtime/Engine/Private/WorldPartition/DataLayer/DataLayerManager.cpp
+++ b/Engine/Source/Runtime/Engine/Private/WorldPartition/DataLayer/DataLayerManager.cpp
@@ -644,10 +644,28 @@ void UDataLayerManager::GetUserLoadedInEditorStates(TArray<FName>& OutDataLayers
 	});
 }
 
+TArray<AWorldDataLayers*> UDataLayerManager::GetActorEditorContextWorldDataLayers() const
+{
+	TArray<AWorldDataLayers*> WorldDataLayersArray;
+	if (AWorldDataLayers* WorldDataLayers = GetWorldDataLayers())
+	{
+		WorldDataLayersArray.Add(WorldDataLayers);
+	}
+	const UWorld* OuterWorld = GetTypedOuter<UWorld>();
+	if (const ULevel* CurrentLevel = (OuterWorld->GetCurrentLevel() && (OuterWorld->GetCurrentLevel() != OuterWorld->PersistentLevel)) ? OuterWorld->GetCurrentLevel() : nullptr)
+	{
+		if (AWorldDataLayers* CurrentLevelWorldDataLayers = CurrentLevel->GetWorldDataLayers())
+		{
+			WorldDataLayersArray.Add(CurrentLevelWorldDataLayers);
+		}
+	}
+	return WorldDataLayersArray;
+}
+
 void UDataLayerManager::PushActorEditorContext() const
 {
 	++DataLayerActorEditorContextID;
-	if (AWorldDataLayers* WorldDataLayers = GetWorldDataLayers())
+	for (AWorldDataLayers* WorldDataLayers : GetActorEditorContextWorldDataLayers())
 	{
 		WorldDataLayers->PushActorEditorContext(DataLayerActorEditorContextID);
 	}
@@ -656,7 +674,7 @@ void UDataLayerManager::PushActorEditorContext() const
 void UDataLayerManager::PopActorEditorContext() const
 {
 	check(DataLayerActorEditorContextID > 0);
-	if (AWorldDataLayers* WorldDataLayers = GetWorldDataLayers())
+	for (AWorldDataLayers* WorldDataLayers : GetActorEditorContextWorldDataLayers())
 	{
 		WorldDataLayers->PopActorEditorContext(DataLayerActorEditorContextID);
 	}
@@ -665,9 +683,12 @@ void UDataLayerManager::PopActorEditorContext() const
 
 TArray<UDataLayerInstance*> UDataLayerManager::GetActorEditorContextDataLayers() const
 {
-	static TArray<UDataLayerInstance*> EmptyArray;
-	AWorldDataLayers* WorldDataLayers = GetWorldDataLayers();
-	return WorldDataLayers ? WorldDataLayers->GetActorEditorContextDataLayers() : EmptyArray;
+	TArray<UDataLayerInstance*> ActorEditorContextDataLayers;
+	for (AWorldDataLayers* WorldDataLayers : GetActorEditorContextWorldDataLayers())
+	{
+		ActorEditorContextDataLayers.Append(WorldDataLayers->GetActorEditorContextDataLayers());
+	}
+	return ActorEditorContextDataLayers;
 }
 
 uint32 UDataLayerManager::GetDataLayerEditorContextHash() const
diff --git a/Engine/Source/Runtime/Engine/Private/WorldPartition/DataLayer/WorldDataLayers.cpp b/Engine/Source/Runtime/Engine/Private/WorldPartition/DataLayer/WorldDataLayers.cpp
--- a/Engine/Source/Runtime/Engine/Private/WorldPartition/DataLayer/WorldDataLayers.cpp
+++ b/Engine/Source/Runtime/Engine/Private/WorldPartition/DataLayer/WorldDataLayers.cpp
@@ -850,6 +850,18 @@ bool AWorldDataLayers::IsSubWorldDataLayers() const
 	return ActorWorld != nullptr && OuterWorld != nullptr && (OuterWorld->GetFName() != ActorWorld->GetFName());
 }
 
+bool AWorldDataLayers::IsReadOnly() const
+{
+	if (IsSubWorldDataLayers())
+	{
+		const UWorld* ActorWorld = GetWorld();
+		const ULevel* CurrentLevel = (ActorWorld && ActorWorld->GetCurrentLevel() && !ActorWorld->GetCurrentLevel()->IsPersistentLevel()) ? ActorWorld->GetCurrentLevel() : nullptr;
+		const bool bIsCurrentLevelWorldDataLayers = CurrentLevel && (CurrentLevel == GetLevel());
+		return !bIsCurrentLevelWorldDataLayers;
+	}
+	return false;
+}
+
 bool AWorldDataLayers::IsRuntimeRelevant() const
 {
 	// This function is only meant to be called on Game worlds
diff --git a/Engine/Source/Runtime/Engine/Public/WorldPartition/DataLayer/DataLayerManager.h b/Engine/Source/Runtime/Engine/Public/WorldPartition/DataLayer/DataLayerManager.h
--- a/Engine/Source/Runtime/Engine/Public/WorldPartition/DataLayer/DataLayerManager.h
+++ b/Engine/Source/Runtime/Engine/Public/WorldPartition/DataLayer/DataLayerManager.h
@@ -137,6 +137,7 @@ private:
 	void PushActorEditorContext() const;
 	void PopActorEditorContext() const;
 	TArray<UDataLayerInstance*> GetActorEditorContextDataLayers() const;
+	TArray<AWorldDataLayers*> GetActorEditorContextWorldDataLayers() const;
 	uint32 GetDataLayerEditorContextHash() const;
 	//~ End Editor Context
 
diff --git a/Engine/Source/Runtime/Engine/Public/WorldPartition/DataLayer/WorldDataLayers.h b/Engine/Source/Runtime/Engine/Public/WorldPartition/DataLayer/WorldDataLayers.h
--- a/Engine/Source/Runtime/Engine/Public/WorldPartition/DataLayer/WorldDataLayers.h
+++ b/Engine/Source/Runtime/Engine/Public/WorldPartition/DataLayer/WorldDataLayers.h
@@ -88,6 +88,7 @@ public:
 	template<class T>
 	void OverwriteDataLayerRuntimeStates(const TArray<T>* InActiveDataLayers, const TArray<T>* InLoadedDataLayers );
 
+	bool IsReadOnly() const;
 	bool IsSubWorldDataLayers() const;
 	bool IsRuntimeRelevant() const;
 #endif
