DataLayer Outliner: Fix WorldDataLayer item not being displayed when World doesn't yet have DL Instances

diff --git a/Engine/Source/Editor/DataLayerEditor/Private/DataLayer/DataLayerHierarchy.cpp b/Engine/Source/Editor/DataLayerEditor/Private/DataLayer/DataLayerHierarchy.cpp
--- a/Engine/Source/Editor/DataLayerEditor/Private/DataLayer/DataLayerHierarchy.cpp
+++ b/Engine/Source/Editor/DataLayerEditor/Private/DataLayer/DataLayerHierarchy.cpp
@@ -172,6 +172,14 @@ void FDataLayerHierarchy::CreateItems(TArray<FSceneOutlinerTreeItemPtr>& OutItem
 	{
 		if (InDataLayerManager)
 		{
+			if (AWorldDataLayers* WorldDataLayers = InDataLayerManager->GetWorldDataLayers())
+			{
+				if (FSceneOutlinerTreeItemPtr WorldDataLayersItem = Mode->CreateItemFor<FWorldDataLayersTreeItem>(WorldDataLayers))
+				{
+					OutItems.Add(WorldDataLayersItem);
+				}
+			}
+
 			InDataLayerManager->ForEachDataLayerInstance([this, &OutItems, IsDataLayerShown](UDataLayerInstance* DataLayerInstance)
 			{
 				if (IsDataLayerShown(DataLayerInstance))
