LevelInstance: Fix Skipped actor descs being displayed in outliner since CL 24422197

#rnx

diff --git a/Engine/Source/Editor/DataLayerEditor/Private/DataLayer/DataLayerHierarchy.cpp b/Engine/Source/Editor/DataLayerEditor/Private/DataLayer/DataLayerHierarchy.cpp
--- a/Engine/Source/Editor/DataLayerEditor/Private/DataLayer/DataLayerHierarchy.cpp
+++ b/Engine/Source/Editor/DataLayerEditor/Private/DataLayer/DataLayerHierarchy.cpp
@@ -285,7 +285,7 @@ void FDataLayerHierarchy::CreateItems(TArray<FSceneOutlinerTreeItemPtr>& OutItem
 						// Create an FDataLayerActorDescTreeItem for each unloaded actor of this WorldPartition
 						FWorldPartitionHelpers::ForEachActorDesc(WorldPartition, [this, IsDataLayerShown, DataLayerManager, CurrentLevel, &WorldToLevelDataLayerMap, &OutItems](const FWorldPartitionActorDesc* ActorDesc)
 						{
-							if (ActorDesc != nullptr && !ActorDesc->IsLoaded())
+							if (ActorDesc != nullptr && !ActorDesc->IsLoaded() && FActorDescTreeItem::ShouldDisplayInOutliner(ActorDesc))
 							{
 								for (const FName& DataLayerInstanceName : ActorDesc->GetDataLayerInstanceNames())
 								{
@@ -561,7 +561,7 @@ void FDataLayerHierarchy::OnLoadedActorRemoved(AActor& InActor)
 			EventData.Items.Reserve(DataLayerInstances.Num());
 			for (const UDataLayerInstance* DataLayerInstance : DataLayerInstances)
 			{
-				if (FWorldPartitionActorDesc* ActorDesc = WorldPartition->GetActorDesc(InActor.GetActorGuid()))
+				if (FWorldPartitionActorDesc* ActorDesc = WorldPartition->GetActorDesc(InActor.GetActorGuid()); FActorDescTreeItem::ShouldDisplayInOutliner(ActorDesc))
 				{
 					EventData.Items.Add(Mode->CreateItemFor<FDataLayerActorDescTreeItem>(FDataLayerActorDescTreeItemData(ActorDesc->GetGuid(), ActorDesc->GetContainer(), const_cast<UDataLayerInstance*>(DataLayerInstance))));
 				}
@@ -573,7 +573,7 @@ void FDataLayerHierarchy::OnLoadedActorRemoved(AActor& InActor)
 
 void FDataLayerHierarchy::OnActorDescAdded(FWorldPartitionActorDesc* InActorDesc)
 {
-	if (!bShowUnloadedActors || !InActorDesc || InActorDesc->IsLoaded(true))
+	if (!bShowUnloadedActors || !InActorDesc || InActorDesc->IsLoaded(true) || !FActorDescTreeItem::ShouldDisplayInOutliner(InActorDesc))
 	{
 		return;
 	}
diff --git a/Engine/Source/Editor/SceneOutliner/Private/ActorDescTreeItem.cpp b/Engine/Source/Editor/SceneOutliner/Private/ActorDescTreeItem.cpp
--- a/Engine/Source/Editor/SceneOutliner/Private/ActorDescTreeItem.cpp
+++ b/Engine/Source/Editor/SceneOutliner/Private/ActorDescTreeItem.cpp
@@ -257,6 +257,11 @@ FSceneOutlinerTreeItemID FActorDescTreeItem::ComputeTreeItemID(FGuid InActorGuid
 	return FSceneOutlinerTreeItemID(MD5HashToGuid(MD5Hash));
 }
 
+bool FActorDescTreeItem::ShouldDisplayInOutliner(const FWorldPartitionActorDesc* ActorDesc)
+{
+	return ActorDesc && (ActorDesc->GetActorIsRuntimeOnly() || ActorDesc->IsEditorRelevant());
+}
+
 FSceneOutlinerTreeItemID FActorDescTreeItem::GetID() const
 {
 	return ID;
diff --git a/Engine/Source/Editor/SceneOutliner/Private/ActorHierarchy.cpp b/Engine/Source/Editor/SceneOutliner/Private/ActorHierarchy.cpp
--- a/Engine/Source/Editor/SceneOutliner/Private/ActorHierarchy.cpp
+++ b/Engine/Source/Editor/SceneOutliner/Private/ActorHierarchy.cpp
@@ -458,7 +458,7 @@ void FActorHierarchy::CreateWorldChildren(UWorld* World, TArray<FSceneOutlinerTr
 				}
 				FWorldPartitionHelpers::ForEachActorDesc(WorldPartition, [this, &OutItems](const FWorldPartitionActorDesc* ActorDesc)
 				{
-					if (ActorDesc != nullptr && !ActorDesc->IsLoaded(true))
+					if (ActorDesc != nullptr && !ActorDesc->IsLoaded(true) && FActorDescTreeItem::ShouldDisplayInOutliner(ActorDesc))
 					{
 						if (const FSceneOutlinerTreeItemPtr ActorDescItem = Mode->CreateItemFor<FActorDescTreeItem>(FActorDescTreeItem(ActorDesc->GetGuid(), ActorDesc->GetContainer())))
 						{
@@ -665,7 +665,7 @@ void FActorHierarchy::OnLoadedActorRemoved(AActor& InActor)
 		if (UWorldPartition* WorldPartition = FWorldPartitionHelpers::GetWorldPartition(&InActor))
 		{
 			const FGuid& ActorGuid = InActor.GetActorGuid();
-			if (FWorldPartitionActorDesc* ActorDesc = WorldPartition->GetActorDesc(ActorGuid))
+			if (FWorldPartitionActorDesc* ActorDesc = WorldPartition->GetActorDesc(ActorGuid); FActorDescTreeItem::ShouldDisplayInOutliner(ActorDesc))
 			{
 				FSceneOutlinerHierarchyChangedData EventData;
 				EventData.Type = FSceneOutlinerHierarchyChangedData::Added;
@@ -678,7 +678,7 @@ void FActorHierarchy::OnLoadedActorRemoved(AActor& InActor)
 
 void FActorHierarchy::OnActorDescAdded(FWorldPartitionActorDesc* ActorDesc)
 {
-	if (bShowingUnloadedActors && ActorDesc && !ActorDesc->IsLoaded(true))
+	if (bShowingUnloadedActors && ActorDesc && !ActorDesc->IsLoaded(true) && FActorDescTreeItem::ShouldDisplayInOutliner(ActorDesc))
 	{
 		FSceneOutlinerHierarchyChangedData EventData;
 		EventData.Type = FSceneOutlinerHierarchyChangedData::Added;
diff --git a/Engine/Source/Editor/SceneOutliner/Public/ActorDescTreeItem.h b/Engine/Source/Editor/SceneOutliner/Public/ActorDescTreeItem.h
--- a/Engine/Source/Editor/SceneOutliner/Public/ActorDescTreeItem.h
+++ b/Engine/Source/Editor/SceneOutliner/Public/ActorDescTreeItem.h
@@ -36,6 +36,7 @@ public:
 	FSceneOutlinerTreeItemID ID;
 
 	static FSceneOutlinerTreeItemID ComputeTreeItemID(FGuid InActorGuid, UActorDescContainer* InContainer);
+	static bool ShouldDisplayInOutliner(const FWorldPartitionActorDesc* ActorDesc);
 
 	/** Static type identifier for this tree item class */
 	static const FSceneOutlinerTreeItemType Type;
