DataLayerEditorSubsystem: Update Visibility on initialized World Partition only (Optim)

#rnx

diff --git a/Engine/Source/Editor/DataLayerEditor/Private/DataLayer/DataLayerEditorSubsystem.cpp b/Engine/Source/Editor/DataLayerEditor/Private/DataLayer/DataLayerEditorSubsystem.cpp
--- a/Engine/Source/Editor/DataLayerEditor/Private/DataLayer/DataLayerEditorSubsystem.cpp
+++ b/Engine/Source/Editor/DataLayerEditor/Private/DataLayer/DataLayerEditorSubsystem.cpp
@@ -430,13 +430,14 @@ bool UDataLayerEditorSubsystem::IsActorValidForDataLayer(AActor* Actor)
 
 void UDataLayerEditorSubsystem::OnWorldPartitionInitialized(UWorldPartition* InWorldPartition)
 {
-	InWorldPartition->GetTypedOuter<UWorld>()->PersistentLevel->OnLoadedActorAddedToLevelEvent.AddUObject(this, &UDataLayerEditorSubsystem::OnLoadedActorAddedToLevel);
-	UpdateAllActorsVisibility(true, true);
+	ULevel* WorldPartitionLevel = InWorldPartition->GetTypedOuter<ULevel>();
+	WorldPartitionLevel->OnLoadedActorAddedToLevelEvent.AddUObject(this, &UDataLayerEditorSubsystem::OnLoadedActorAddedToLevel);
+	UpdateAllActorsVisibility(true, true, WorldPartitionLevel);
 }
 
 void UDataLayerEditorSubsystem::OnWorldPartitionUninitialized(UWorldPartition* InWorldPartition)
 {
-	InWorldPartition->GetTypedOuter<UWorld>()->PersistentLevel->OnLoadedActorAddedToLevelEvent.RemoveAll(this);
+	InWorldPartition->GetTypedOuter<ULevel>()->OnLoadedActorAddedToLevelEvent.RemoveAll(this);
 }
 
 void UDataLayerEditorSubsystem::OnLoadedActorAddedToLevel(AActor& InActor)
@@ -839,21 +840,41 @@ bool UDataLayerEditorSubsystem::UpdateActorVisibility(AActor* Actor, bool& bOutS
 }
 
 bool UDataLayerEditorSubsystem::UpdateAllActorsVisibility(const bool bNotifySelectionChange, const bool bRedrawViewports)
+{
+	return UpdateAllActorsVisibility(bNotifySelectionChange, bRedrawViewports, nullptr);
+}
+
+bool UDataLayerEditorSubsystem::UpdateAllActorsVisibility(const bool bNotifySelectionChange, const bool bRedrawViewports, ULevel* InLevel)
 {
 	TRACE_CPUPROFILER_EVENT_SCOPE(UDataLayerEditorSubsystem::UpdateAllActorsVisibility);
 
 	bool bSelectionChanged = false;
 	bool bChangesOccurred = false;
 
-	if (UWorld* World = GetWorld())
+	auto UpdateActorVisibilityLambda = [this, &bSelectionChanged, &bChangesOccurred](AActor* Actor)
 	{
-		for (AActor* Actor : FActorRange(World))
+		if (Actor)
 		{
 			bool bActorModified = false;
 			bool bActorSelectionChanged = false;
 			bChangesOccurred |= UpdateActorVisibility(Actor, bActorSelectionChanged, bActorModified, /*bActorNotifySelectionChange*/false, /*bActorRedrawViewports*/false);
 			bSelectionChanged |= bActorSelectionChanged;
 		}
+	};
+
+	if (InLevel)
+	{
+		for (AActor* Actor : InLevel->Actors)
+		{
+			UpdateActorVisibilityLambda(Actor);
+		}
+	}
+	else if (UWorld* World = GetWorld())
+	{
+		for (AActor* Actor : FActorRange(World))
+		{
+			UpdateActorVisibilityLambda(Actor);
+		}
 	}
 
 	if (bNotifySelectionChange && bSelectionChanged)
diff --git a/Engine/Source/Editor/DataLayerEditor/Public/DataLayer/DataLayerEditorSubsystem.h b/Engine/Source/Editor/DataLayerEditor/Public/DataLayer/DataLayerEditorSubsystem.h
--- a/Engine/Source/Editor/DataLayerEditor/Public/DataLayer/DataLayerEditorSubsystem.h
+++ b/Engine/Source/Editor/DataLayerEditor/Public/DataLayer/DataLayerEditorSubsystem.h
@@ -723,6 +723,8 @@ private:
 	void RebuildSelectedDataLayersFromEditorSelection();
 	const TSet<TWeakObjectPtr<const UDataLayerInstance>>& GetSelectedDataLayersFromEditorSelection() const;
 
+	bool UpdateAllActorsVisibility(const bool bNotifySelectionChange, const bool bRedrawViewports, ULevel* InLevel);
+
 	/** Contains Data Layers that contain actors that are part of the editor selection */
 	mutable TSet<TWeakObjectPtr<const UDataLayerInstance>> SelectedDataLayersFromEditorSelection;
 
