DataLayer: Refactor some relabel code to support DataLayerInstanceWithAsset renaming (when they are private)
- Renamed to ShortName to simplify
- Moved into FDataLayerUtils
- New code isn't deprecated as support for renaming private DLs will be added
- New code is independant from deprecated code 5.1/5.0 that can be removed once we update our internal projects (UE-176057)

#rnx

diff --git a/Engine/Source/Editor/DataLayerEditor/Private/DataLayer/DataLayerEditorSubsystem.cpp b/Engine/Source/Editor/DataLayerEditor/Private/DataLayer/DataLayerEditorSubsystem.cpp
--- a/Engine/Source/Editor/DataLayerEditor/Private/DataLayer/DataLayerEditorSubsystem.cpp
+++ b/Engine/Source/Editor/DataLayerEditor/Private/DataLayer/DataLayerEditorSubsystem.cpp
@@ -48,6 +48,7 @@
 #include "WorldPartition/DataLayer/DataLayerManager.h"
 #include "WorldPartition/DataLayer/DeprecatedDataLayerInstance.h"
 #include "WorldPartition/DataLayer/WorldDataLayers.h"
+#include "WorldPartition/DataLayer/DataLayerUtils.h"
 #include "WorldPartition/WorldPartitionActorDescViewProxy.h"
 #include "WorldPartition/WorldPartition.h"
 
@@ -1333,6 +1334,20 @@ const TSet<TWeakObjectPtr<const UDataLayerInstance>>& UDataLayerEditorSubsystem:
 	return SelectedDataLayersFromEditorSelection;
 }
 
+bool UDataLayerEditorSubsystem::SetDataLayerShortName(UDataLayerInstance* DataLayerInstance, const FString& InNewShortName)
+{
+	if (DataLayerInstance->CanEditDataLayerShortName())
+	{
+		if (FDataLayerUtils::SetDataLayerShortName(DataLayerInstance, InNewShortName))
+		{
+			BroadcastDataLayerChanged(EDataLayerAction::Rename, DataLayerInstance, "DataLayerShortName");
+			return true;
+		}
+	}
+
+	return false;
+}
+
 //~ Begin Deprecated
 
 PRAGMA_DISABLE_DEPRECATION_WARNINGS
@@ -1351,12 +1366,6 @@ bool UDataLayerEditorSubsystem::RenameDataLayer(UDataLayerInstance* DataLayerIns
 	return false;
 }
 
-bool UDataLayerEditorSubsystem::TryGetDataLayerFromLabel(const FName& DataLayerLabel, UDataLayerInstance*& OutDataLayer) const
-{
-	OutDataLayer = GetDataLayerFromLabel(DataLayerLabel);
-	return (OutDataLayer != nullptr);
-}
-
 UDataLayerInstance* UDataLayerEditorSubsystem::GetDataLayerFromLabel(const FName& DataLayerLabel) const
 {
 	if (AWorldDataLayers* WorldDataLayers = GetWorld()->GetWorldDataLayers())
diff --git a/Engine/Source/Editor/DataLayerEditor/Private/DataLayer/DataLayerMode.cpp b/Engine/Source/Editor/DataLayerEditor/Private/DataLayer/DataLayerMode.cpp
--- a/Engine/Source/Editor/DataLayerEditor/Private/DataLayer/DataLayerMode.cpp
+++ b/Engine/Source/Editor/DataLayerEditor/Private/DataLayer/DataLayerMode.cpp
@@ -325,7 +325,7 @@ bool FDataLayerMode::CanRenameItem(const ISceneOutlinerTreeItem& Item) const
 	if (Item.IsValid() && (Item.IsA<FDataLayerTreeItem>()))
 	{
 		FDataLayerTreeItem* DataLayerTreeItem = (FDataLayerTreeItem*)&Item;
-		return !DataLayerTreeItem->GetDataLayer()->IsLocked() && DataLayerTreeItem->GetDataLayer()->SupportRelabeling();
+		return !DataLayerTreeItem->GetDataLayer()->IsLocked() && DataLayerTreeItem->GetDataLayer()->CanEditDataLayerShortName();
 	}
 	return false;
 }
diff --git a/Engine/Source/Editor/DataLayerEditor/Private/DataLayer/DataLayerNameEditSink.h b/Engine/Source/Editor/DataLayerEditor/Private/DataLayer/DataLayerNameEditSink.h
--- a/Engine/Source/Editor/DataLayerEditor/Private/DataLayer/DataLayerNameEditSink.h
+++ b/Engine/Source/Editor/DataLayerEditor/Private/DataLayer/DataLayerNameEditSink.h
@@ -19,7 +19,7 @@ class FDataLayerNameEditSink : public UE::EditorWidgets::IObjectNameEditSink
 	virtual FText GetObjectDisplayName(UObject* Object) const override
 	{
 		UDataLayerInstance* DataLayerInstance = CastChecked<UDataLayerInstance>(Object);
-		if (DataLayerInstance->SupportRelabeling())
+		if (DataLayerInstance->CanEditDataLayerShortName())
 		{
 			return FText::Format(FText::FromString("{0}"), FText::FromString(DataLayerInstance->GetDataLayerShortName()));
 		}
@@ -35,21 +35,19 @@ class FDataLayerNameEditSink : public UE::EditorWidgets::IObjectNameEditSink
 	virtual bool IsObjectDisplayNameReadOnly(UObject* Object) const override
 	{
 		UDataLayerInstance* DataLayerInstance = CastChecked<UDataLayerInstance>(Object);
-		return !DataLayerInstance->SupportRelabeling() || DataLayerInstance->IsLocked();
+		return !DataLayerInstance->CanEditDataLayerShortName() || DataLayerInstance->IsLocked();
 	};
 
 	virtual bool SetObjectDisplayName(UObject* Object, FString DisplayName) override
 	{
 		UDataLayerInstance* DataLayerInstance = CastChecked<UDataLayerInstance>(Object);
-		if(DataLayerInstance->SupportRelabeling())
+		if(DataLayerInstance->CanEditDataLayerShortName())
 		{
 			if (!DisplayName.Equals(DataLayerInstance->GetDataLayerShortName(), ESearchCase::CaseSensitive))
 			{
 				const FScopedTransaction Transaction(LOCTEXT("DataLayerNameEditSinkRenameDataLayerTransaction", "Rename Data Layer"));
 
-				PRAGMA_DISABLE_DEPRECATION_WARNINGS
-				return UDataLayerEditorSubsystem::Get()->RenameDataLayer(DataLayerInstance, *DisplayName);
-				PRAGMA_ENABLE_DEPRECATION_WARNINGS
+				return UDataLayerEditorSubsystem::Get()->SetDataLayerShortName(DataLayerInstance, DisplayName);
 			}
 		}
 		
diff --git a/Engine/Source/Editor/DataLayerEditor/Private/DataLayer/SDataLayerTreeLabel.cpp b/Engine/Source/Editor/DataLayerEditor/Private/DataLayer/SDataLayerTreeLabel.cpp
--- a/Engine/Source/Editor/DataLayerEditor/Private/DataLayer/SDataLayerTreeLabel.cpp
+++ b/Engine/Source/Editor/DataLayerEditor/Private/DataLayer/SDataLayerTreeLabel.cpp
@@ -35,6 +35,8 @@
 #include "Widgets/Text/STextBlock.h"
 #include "Widgets/Views/STableRow.h"
 #include "WorldPartition/DataLayer/DataLayerInstance.h"
+#include "WorldPartition/DataLayer/DataLayerManager.h"
+#include "WorldPartition/DataLayer/DataLayerUtils.h"
 #include "WorldPartition/DataLayer/DataLayerInstanceWithAsset.h"
 
 struct FSlateBrush;
@@ -69,7 +71,7 @@ void SDataLayerTreeLabel::Construct(const FArguments& InArgs, FDataLayerTreeItem
 		.IsSelected(FIsSelected::CreateSP(&InRow, &STableRow<FSceneOutlinerTreeItemPtr>::IsSelectedExclusively))
 		.IsReadOnly_Lambda([Item = DataLayerItem.AsShared(), this]()
 		{
-			return DataLayerPtr == nullptr || !DataLayerPtr->SupportRelabeling() || !CanExecuteRenameRequest(Item.Get());
+			return DataLayerPtr == nullptr || !DataLayerPtr->CanEditDataLayerShortName() || !CanExecuteRenameRequest(Item.Get());
 		})
 	]
 
@@ -300,20 +302,24 @@ bool SDataLayerTreeLabel::OnVerifyItemLabelChanged(const FText& InLabel, FText&
 		return false;
 	}
 
-	UDataLayerInstance* FoundDataLayerInstance;
-	
-	PRAGMA_DISABLE_DEPRECATION_WARNINGS
-	if (UDataLayerEditorSubsystem::Get()->TryGetDataLayerFromLabel(*InLabel.ToString(), FoundDataLayerInstance) && FoundDataLayerInstance != DataLayerPtr.Get())
+	UDataLayerInstance* DataLayerInstance = DataLayerPtr.Get();
+	if (!DataLayerInstance->CanEditDataLayerShortName())
 	{
-		OutErrorMessage = LOCTEXT("RenameFailed_AlreadyExists", "This Data Layer already exists");
+		OutErrorMessage = LOCTEXT("RenameFailed_NotPermittedOnDataLayer", "This Data Layer does not support renaming");
 		return false;
 	}
-	PRAGMA_ENABLE_DEPRECATION_WARNINGS
 
-	if (FoundDataLayerInstance != nullptr && !FoundDataLayerInstance->SupportRelabeling())
+	UDataLayerManager* DataLayerManager = UDataLayerManager::GetDataLayerManager(DataLayerPtr.Get());
+	check(DataLayerManager);
+
+	TSet<UDataLayerInstance*> OutDataLayersWithShortName;
+	if (FDataLayerUtils::FindDataLayerByShortName(DataLayerManager, InLabel.ToString(), OutDataLayersWithShortName))
 	{
-		OutErrorMessage = LOCTEXT("RenameFailed_NotPermittedOnDataLayer", "This Data Layer does not support renaming");
-		return false;
+		if (!OutDataLayersWithShortName.Contains(DataLayerPtr.Get()))
+		{
+			OutErrorMessage = LOCTEXT("RenameFailed_AlreadyExists", "This Data Layer already exists");
+			return false;
+		}
 	}
 
 	return true;
@@ -322,14 +328,12 @@ bool SDataLayerTreeLabel::OnVerifyItemLabelChanged(const FText& InLabel, FText&
 void SDataLayerTreeLabel::OnLabelCommitted(const FText& InLabel, ETextCommit::Type InCommitInfo)
 {
 	UDataLayerInstance* DataLayerInstance = DataLayerPtr.Get();
-	check(DataLayerInstance->SupportRelabeling());
+	check(DataLayerInstance->CanEditDataLayerShortName());
 	if (!InLabel.ToString().Equals(DataLayerInstance->GetDataLayerShortName(), ESearchCase::CaseSensitive))
 	{
 		const FScopedTransaction Transaction(LOCTEXT("SceneOutlinerRenameDataLayerTransaction", "Rename Data Layer"));
 
-		PRAGMA_DISABLE_DEPRECATION_WARNINGS
-		UDataLayerEditorSubsystem::Get()->RenameDataLayer(DataLayerInstance, *InLabel.ToString());
-		PRAGMA_ENABLE_DEPRECATION_WARNINGS
+		UDataLayerEditorSubsystem::Get()->SetDataLayerShortName(DataLayerInstance, InLabel.ToString());
 
 		if (WeakSceneOutliner.IsValid())
 		{
diff --git a/Engine/Source/Editor/DataLayerEditor/Public/DataLayer/DataLayerEditorSubsystem.h b/Engine/Source/Editor/DataLayerEditor/Public/DataLayer/DataLayerEditorSubsystem.h
--- a/Engine/Source/Editor/DataLayerEditor/Public/DataLayer/DataLayerEditorSubsystem.h
+++ b/Engine/Source/Editor/DataLayerEditor/Public/DataLayer/DataLayerEditorSubsystem.h
@@ -615,6 +615,11 @@ public:
 	*/
 	bool HasDeprecatedDataLayers() const;
 
+	/**
+	 * Assign new unique short name to DataLayerInstance if it supports it. 
+	 */
+	bool SetDataLayerShortName(UDataLayerInstance* InDataLayerInstance, const FString& InNewShortName);
+
 	//~ Begin Deprecated
 
 	PRAGMA_DISABLE_DEPRECATION_WARNINGS
@@ -632,10 +637,7 @@ public:
 	UE_DEPRECATED(5.0, "Per-view Data Layer visibility was removed.")
 	UFUNCTION(BlueprintCallable, Category = DataLayers, meta = (DeprecatedFunction, DeprecationMessage = "Per-view Data Layer visibility was removed."))
 	void UpdateActorAllViewsVisibility(AActor* Actor) {}
-
-	UE_DEPRECATED(5.0, "Per-view Data Layer visibility was removed.")
-	void RemoveViewFromActorViewVisibility(FLevelEditorViewportClient* ViewportClient);
-
+		
 	UE_DEPRECATED(5.0, "Use SetDataLayerIsLoadedInEditor() instead.")
 	UFUNCTION(BlueprintCallable, Category = DataLayers, meta = (DeprecatedFunction, DeprecationMessage = "Use SetDataLayerIsLoadedInEditor instead"))
 	bool SetDataLayerIsDynamicallyLoadedInEditor(UDEPRECATED_DataLayer* DataLayer, const bool bIsLoadedInEditor, const bool bIsFromUserChange) { return false; }
@@ -653,10 +655,7 @@ public:
 	bool ToggleDataLayersIsDynamicallyLoadedInEditor(const TArray<UDEPRECATED_DataLayer*>& DataLayers, const bool bIsFromUserChange) { return false; }
 
 	PRAGMA_ENABLE_DEPRECATION_WARNINGS
-
-	UE_DEPRECATED(5.1, "Use GetDataLayerInstance instead")
-	bool TryGetDataLayerFromLabel(const FName& DataLayerLabel, UDataLayerInstance*& OutDataLayer) const;
-
+			
 	UE_DEPRECATED(5.1, "Use GetDataLayerInstance instead")
 	UFUNCTION(BlueprintCallable, Category = DataLayers)
 	UDataLayerInstance* GetDataLayer(const FActorDataLayer& ActorDataLayer) const;
diff --git a/Engine/Source/Runtime/Engine/Private/WorldPartition/DataLayer/DataLayer.cpp b/Engine/Source/Runtime/Engine/Private/WorldPartition/DataLayer/DataLayer.cpp
--- a/Engine/Source/Runtime/Engine/Private/WorldPartition/DataLayer/DataLayer.cpp
+++ b/Engine/Source/Runtime/Engine/Private/WorldPartition/DataLayer/DataLayer.cpp
@@ -45,7 +45,7 @@ void UDEPRECATED_DataLayer::PostLoad()
 	bIsVisible = bIsInitiallyVisible;
 
 	// Sanitize Label
-	DataLayerLabel = FDataLayerUtils::GetSanitizedDataLayerLabel(DataLayerLabel);
+	DataLayerLabel = *FDataLayerUtils::GetSanitizedDataLayerShortName(DataLayerLabel.ToString());
 
 	if (DebugColor == FColor::Black)
 	{
diff --git a/Engine/Source/Runtime/Engine/Private/WorldPartition/DataLayer/DataLayerUtils.cpp b/Engine/Source/Runtime/Engine/Private/WorldPartition/DataLayer/DataLayerUtils.cpp
--- a/Engine/Source/Runtime/Engine/Private/WorldPartition/DataLayer/DataLayerUtils.cpp
+++ b/Engine/Source/Runtime/Engine/Private/WorldPartition/DataLayer/DataLayerUtils.cpp
@@ -163,4 +163,47 @@ bool FDataLayerUtils::AreWorldDataLayersActorDescsSane(const TArray<const FWorld
 	return Algo::AllOf(InWorldDataLayersActorDescs, [](const FWorldDataLayersActorDesc* WorldDataLayerActorDescs) { return WorldDataLayerActorDescs->IsValid(); });
 }
 
+bool FDataLayerUtils::SetDataLayerShortName(UDataLayerInstance* InDataLayerInstance, const FString& InNewShortName)
+{
+	check(InDataLayerInstance->CanEditDataLayerShortName());
+	FString UniqueShortName = FDataLayerUtils::GenerateUniqueDataLayerShortName(UDataLayerManager::GetDataLayerManager(InDataLayerInstance), InNewShortName);
+	if (InDataLayerInstance->GetDataLayerShortName() != UniqueShortName)
+	{
+		InDataLayerInstance->Modify();
+		InDataLayerInstance->PerformSetDataLayerShortName(UniqueShortName);
+		return true;
+	}
+
+	return false;
+}
+
+bool FDataLayerUtils::FindDataLayerByShortName(const UDataLayerManager* InDataLayerManager, const FString& InShortName, TSet<UDataLayerInstance*>&  OutDataLayersWithShortName)
+{
+	OutDataLayersWithShortName.Empty();
+	InDataLayerManager->ForEachDataLayerInstance([&](UDataLayerInstance* DataLayerInstance)
+	{
+		if (DataLayerInstance->GetDataLayerShortName() == InShortName)
+		{
+			OutDataLayersWithShortName.Add(DataLayerInstance);
+		}
+		return true;
+	});
+	return OutDataLayersWithShortName.Num() > 0;
+}
+
+FString FDataLayerUtils::GenerateUniqueDataLayerShortName(const UDataLayerManager* InDataLayerManager, const FString& InNewShortName)
+{
+	int32 DataLayerIndex = 0;
+	const FString DataLayerShortNameSanitized = FDataLayerUtils::GetSanitizedDataLayerShortName(InNewShortName);
+	FString UniqueNewDataLayerShortName = DataLayerShortNameSanitized;
+		
+	TSet<UDataLayerInstance*> OutDataLayersWithShortName;
+	while (FindDataLayerByShortName(InDataLayerManager, UniqueNewDataLayerShortName, OutDataLayersWithShortName))
+	{
+		UniqueNewDataLayerShortName = FString::Printf(TEXT("%s%d"), *DataLayerShortNameSanitized, ++DataLayerIndex);
+	}
+		
+	return UniqueNewDataLayerShortName;
+}
+
 #endif
diff --git a/Engine/Source/Runtime/Engine/Private/WorldPartition/DataLayer/DeprecatedDataLayerInstance.cpp b/Engine/Source/Runtime/Engine/Private/WorldPartition/DataLayer/DeprecatedDataLayerInstance.cpp
--- a/Engine/Source/Runtime/Engine/Private/WorldPartition/DataLayer/DeprecatedDataLayerInstance.cpp
+++ b/Engine/Source/Runtime/Engine/Private/WorldPartition/DataLayer/DeprecatedDataLayerInstance.cpp
@@ -30,7 +30,7 @@ void UDeprecatedDataLayerInstance::OnCreated()
 
 	Modify(/*bAlwaysMarkDirty*/false);
 
-	RelabelDataLayer(TEXT("DataLayer"));
+	FDataLayerUtils::SetDataLayerShortName(this, TEXT("DataLayer"));
 
 	DeprecatedDataLayerFName = TEXT("");
 
@@ -84,17 +84,7 @@ bool UDeprecatedDataLayerInstance::PerformRemoveActor(AActor* InActor) const
 
 bool UDeprecatedDataLayerInstance::RelabelDataLayer(FName InDataLayerLabel)
 {
-	FName SanitizedLabel = FDataLayerUtils::GetSanitizedDataLayerLabel(InDataLayerLabel);
-	FName UniqueNewDataLayerLabel = GetOuterAWorldDataLayers()->GenerateUniqueDataLayerLabel(SanitizedLabel);
-	if (Label != UniqueNewDataLayerLabel)
-	{
-		Modify();
-		check(!GetOuterAWorldDataLayers()->GetDataLayerFromLabel(UniqueNewDataLayerLabel));
-		Label = UniqueNewDataLayerLabel;
-		return true;
-	}
-
-	return false;
+	return FDataLayerUtils::SetDataLayerShortName(this, InDataLayerLabel.ToString());
 }
 
 #endif // WITH_EDITOR
diff --git a/Engine/Source/Runtime/Engine/Private/WorldPartition/DataLayer/WorldDataLayers.cpp b/Engine/Source/Runtime/Engine/Private/WorldPartition/DataLayer/WorldDataLayers.cpp
--- a/Engine/Source/Runtime/Engine/Private/WorldPartition/DataLayer/WorldDataLayers.cpp
+++ b/Engine/Source/Runtime/Engine/Private/WorldPartition/DataLayer/WorldDataLayers.cpp
@@ -1001,18 +1001,6 @@ bool AWorldDataLayers::RemoveDataLayer(const UDEPRECATED_DataLayer* InDataLayer)
 	return false;
 }
 
-FName AWorldDataLayers::GenerateUniqueDataLayerLabel(const FName& InDataLayerLabel) const
-{
-	int32 DataLayerIndex = 0;
-	const FName DataLayerLabelSanitized = FDataLayerUtils::GetSanitizedDataLayerLabel(InDataLayerLabel);
-	FName UniqueNewDataLayerLabel = DataLayerLabelSanitized;
-	while (GetDataLayerFromLabel(UniqueNewDataLayerLabel))
-	{
-		UniqueNewDataLayerLabel = FName(*FString::Printf(TEXT("%s%d"), *DataLayerLabelSanitized.ToString(), ++DataLayerIndex));
-	};
-	return UniqueNewDataLayerLabel;
-}
-
 #endif
 
 bool AWorldDataLayers::ContainsDataLayer(const UDEPRECATED_DataLayer* InDataLayer) const
@@ -1022,11 +1010,11 @@ bool AWorldDataLayers::ContainsDataLayer(const UDEPRECATED_DataLayer* InDataLaye
 
 const UDataLayerInstance* AWorldDataLayers::GetDataLayerFromLabel(const FName& InDataLayerLabel) const
 {
-	const FName DataLayerLabelSanitized = FDataLayerUtils::GetSanitizedDataLayerLabel(InDataLayerLabel);
+	const FString DataLayerLabelSanitized = *FDataLayerUtils::GetSanitizedDataLayerShortName(InDataLayerLabel.ToString());
 
 	for (const UDataLayerInstance* DataLayer : DataLayerInstances)
 	{
-		if (FName(DataLayer->GetDataLayerShortName()) == DataLayerLabelSanitized)
+		if (DataLayer->GetDataLayerShortName() == DataLayerLabelSanitized)
 		{
 			return DataLayer;
 		}
diff --git a/Engine/Source/Runtime/Engine/Public/WorldPartition/DataLayer/DataLayerInstance.h b/Engine/Source/Runtime/Engine/Public/WorldPartition/DataLayer/DataLayerInstance.h
--- a/Engine/Source/Runtime/Engine/Public/WorldPartition/DataLayer/DataLayerInstance.h
+++ b/Engine/Source/Runtime/Engine/Public/WorldPartition/DataLayer/DataLayerInstance.h
@@ -101,11 +101,13 @@ public:
 
 	virtual bool Validate(IStreamingGenerationErrorHandler* ErrorHandler) const;
 
-#if DATALAYER_TO_INSTANCE_RUNTIME_CONVERSION_ENABLED
+	UE_DEPRECATED(5.3, "Use CanEditShortName instead")
 	virtual bool SupportRelabeling() const { return false; }
+
+	UE_DEPRECATED(5.3, "Use SetShortName instead")
 	virtual bool RelabelDataLayer(FName NewDataLayerLabel) { return false; }
-#endif // DATALAYER_TO_INSTANCE_RUNTIME_CONVERSION_ENABLED
 
+	virtual bool CanEditDataLayerShortName() const { return false; }
 #endif // WITH_EDITOR
 
 	UFUNCTION(Category = "Data Layer", BlueprintCallable)
@@ -154,6 +156,7 @@ private:
 	void AddChild(UDataLayerInstance* DataLayer);
 	bool SetRuntimeState(EDataLayerRuntimeState InState, bool bInIsRecursive = false) const;
 	friend class UDataLayerManager;
+	friend class FDataLayerUtils;
 
 #if WITH_EDITOR
 	void RemoveChild(UDataLayerInstance* DataLayer);
@@ -163,6 +166,7 @@ protected:
 #if WITH_EDITOR
 	virtual bool PerformAddActor(AActor* InActor) const { return false; }
 	virtual bool PerformRemoveActor(AActor* InActor) const { return false;  }
+	virtual void PerformSetDataLayerShortName(const FString& InNewShortName) {}
 #endif
 
 #if WITH_EDITORONLY_DATA
diff --git a/Engine/Source/Runtime/Engine/Public/WorldPartition/DataLayer/DataLayerUtils.h b/Engine/Source/Runtime/Engine/Public/WorldPartition/DataLayer/DataLayerUtils.h
--- a/Engine/Source/Runtime/Engine/Public/WorldPartition/DataLayer/DataLayerUtils.h
+++ b/Engine/Source/Runtime/Engine/Public/WorldPartition/DataLayer/DataLayerUtils.h
@@ -38,13 +38,16 @@ public:
 	static TArray<const FWorldDataLayersActorDesc*> FindWorldDataLayerActorDescs(const FActorDescViewMap& ActorDescViewMap);
 
 	static bool AreWorldDataLayersActorDescsSane(const TArray<const FWorldDataLayersActorDesc*>& InWorldDataLayersActorDescs);
+
+	static FString GenerateUniqueDataLayerShortName(const UDataLayerManager* InDataLayerManager, const FString& InNewShortName);
+	
+	static bool SetDataLayerShortName(UDataLayerInstance* InDataLayerInstance, const FString& InNewShortName);
+
+	static bool FindDataLayerByShortName(const UDataLayerManager* InDataLayerManager, const FString& InShortName, TSet<UDataLayerInstance*>& OutDataLayerInstances);
 #endif
 
-#if DATALAYER_TO_INSTANCE_RUNTIME_CONVERSION_ENABLED
-	UE_DEPRECATED(5.1, "Label usage is deprecated.")
-	static FName GetSanitizedDataLayerLabel(FName InDataLayerLabel)
+	static FString GetSanitizedDataLayerShortName(FString InShortName)
 	{
-		return FName(InDataLayerLabel.ToString().TrimStartAndEnd().Replace(TEXT("\""), TEXT("")));
+		return InShortName.TrimStartAndEnd().Replace(TEXT("\""), TEXT(""));
 	}
-#endif
 };
\ No newline at end of file
diff --git a/Engine/Source/Runtime/Engine/Public/WorldPartition/DataLayer/DeprecatedDataLayerInstance.h b/Engine/Source/Runtime/Engine/Public/WorldPartition/DataLayer/DeprecatedDataLayerInstance.h
--- a/Engine/Source/Runtime/Engine/Public/WorldPartition/DataLayer/DeprecatedDataLayerInstance.h
+++ b/Engine/Source/Runtime/Engine/Public/WorldPartition/DataLayer/DeprecatedDataLayerInstance.h
@@ -50,6 +50,8 @@ private:
 
 	virtual bool SupportRelabeling() const override { return true; }
 	virtual bool RelabelDataLayer(FName NewDataLayerLabel) override;
+	virtual bool CanEditDataLayerShortName() const override { return true; }
+	virtual void PerformSetDataLayerShortName(const FString& InNewShortName) { Label = *InNewShortName; }
 #endif
 
 	virtual FName GetDataLayerFName() const override { return !DeprecatedDataLayerFName.IsNone() ? DeprecatedDataLayerFName : Super::GetDataLayerFName(); }
diff --git a/Engine/Source/Runtime/Engine/Public/WorldPartition/DataLayer/WorldDataLayers.h b/Engine/Source/Runtime/Engine/Public/WorldPartition/DataLayer/WorldDataLayers.h
--- a/Engine/Source/Runtime/Engine/Public/WorldPartition/DataLayer/WorldDataLayers.h
+++ b/Engine/Source/Runtime/Engine/Public/WorldPartition/DataLayer/WorldDataLayers.h
@@ -141,11 +141,6 @@ public:
 
 	PRAGMA_ENABLE_DEPRECATION_WARNINGS
 
-#if WITH_EDITOR
-	UE_DEPRECATED(5.1, "Label usage is deprecated.")
-	FName GenerateUniqueDataLayerLabel(const FName& InDataLayerLabel) const;
-#endif 
-
 	UE_DEPRECATED(5.1, "Use GetDataLayerInstance with FName or UDataLayerAsset instead")
 	const UDataLayerInstance* GetDataLayerInstance(const FActorDataLayer& InActorDataLayer) const;
 
