Scene Outliner: don't sort unloaded actors at the end of the list to avoid messy behavior when changing the loaded state of an actor.

diff --git a/Engine/Source/Editor/DataLayerEditor/Private/DataLayer/DataLayerMode.cpp b/Engine/Source/Editor/DataLayerEditor/Private/DataLayer/DataLayerMode.cpp
--- a/Engine/Source/Editor/DataLayerEditor/Private/DataLayer/DataLayerMode.cpp
+++ b/Engine/Source/Editor/DataLayerEditor/Private/DataLayer/DataLayerMode.cpp
@@ -307,14 +307,10 @@ int32 FDataLayerMode::GetTypeSortPriority(const ISceneOutlinerTreeItem& Item) co
 	{
 		return static_cast<int32>(EItemSortOrder::DataLayer);
 	}
-	else if (Item.IsA<FDataLayerActorTreeItem>())
+	else if (Item.IsA<FDataLayerActorTreeItem>() || Item.IsA<FDataLayerActorDescTreeItem>())
 	{
 		return static_cast<int32>(EItemSortOrder::Actor);
 	}
-	else if (Item.IsA<FDataLayerActorDescTreeItem>())
-	{
-		return static_cast<int32_t>(EItemSortOrder::Unloaded);
-	}
 	// Warning: using actor mode with an unsupported item type!
 	check(false);
 	return -1;
diff --git a/Engine/Source/Editor/DataLayerEditor/Private/DataLayer/DataLayerMode.h b/Engine/Source/Editor/DataLayerEditor/Private/DataLayer/DataLayerMode.h
--- a/Engine/Source/Editor/DataLayerEditor/Private/DataLayer/DataLayerMode.h
+++ b/Engine/Source/Editor/DataLayerEditor/Private/DataLayer/DataLayerMode.h
@@ -66,7 +66,6 @@ public:
 		WorldDataLayers = 0,
 		DataLayer = 10,
 		Actor = 20,
-		Unloaded = 30,
 	};
 
 	FDataLayerMode(const FDataLayerModeParams& Params);
diff --git a/Engine/Source/Editor/SceneOutliner/Private/ActorMode.cpp b/Engine/Source/Editor/SceneOutliner/Private/ActorMode.cpp
--- a/Engine/Source/Editor/SceneOutliner/Private/ActorMode.cpp
+++ b/Engine/Source/Editor/SceneOutliner/Private/ActorMode.cpp
@@ -454,14 +454,10 @@ int32 FActorMode::GetTypeSortPriority(const ISceneOutlinerTreeItem& Item) const
 	{
 		return EItemSortOrder::Folder;
 	}
-	else if (Item.IsA<FActorTreeItem>() || Item.IsA<FComponentTreeItem>())
+	else if (Item.IsA<FActorTreeItem>() || Item.IsA<FComponentTreeItem>() || Item.IsA<FActorDescTreeItem>())
 	{
 		return EItemSortOrder::Actor;
 	}
-	else if (Item.IsA<FActorDescTreeItem>())
-	{
-		return EItemSortOrder::Unloaded;
-	}
 
 	// Warning: using actor mode with an unsupported item type!
 	check(false);
diff --git a/Engine/Source/Editor/SceneOutliner/Public/ActorMode.h b/Engine/Source/Editor/SceneOutliner/Public/ActorMode.h
--- a/Engine/Source/Editor/SceneOutliner/Public/ActorMode.h
+++ b/Engine/Source/Editor/SceneOutliner/Public/ActorMode.h
@@ -54,7 +54,7 @@ class SCENEOUTLINER_API FActorMode : public ISceneOutlinerMode
 public:
 	struct EItemSortOrder
 	{
-		enum Type { World = 0, Level = 10, Folder = 20, Actor = 30, Unloaded = 40 };
+		enum Type { World = 0, Level = 10, Folder = 20, Actor = 30 };
 	};
 	
 	FActorMode(const FActorModeParams& Params);
