Expand UDataLayerInstance interface in preparation for ExternalDataLayer.
Allow DataLayerInstance to not support Parenting or being parented.
Allow preventing adding or removing actors from the DataLayerOutliner
Allow preventing adding DataLayerInstance to the actor editor context
Add UDataLayerInstance::IsUserManaged to define if the data layer creation/deletion is managed by the user.

diff --git a/Engine/Source/Editor/DataLayerEditor/Private/DataLayer/DataLayerEditorSubsystem.cpp b/Engine/Source/Editor/DataLayerEditor/Private/DataLayer/DataLayerEditorSubsystem.cpp
--- a/Engine/Source/Editor/DataLayerEditor/Private/DataLayer/DataLayerEditorSubsystem.cpp
+++ b/Engine/Source/Editor/DataLayerEditor/Private/DataLayer/DataLayerEditorSubsystem.cpp
@@ -369,6 +369,7 @@ TSharedRef<SWidget> UDataLayerEditorSubsystem::GetActorEditorContextWidget(UWorl
 
 void UDataLayerEditorSubsystem::AddToActorEditorContext(UDataLayerInstance* InDataLayerInstance)
 {
+	check(InDataLayerInstance->CanBeInActorEditorContext());
 	if (InDataLayerInstance->AddToActorEditorContext())
 	{
 		BroadcastDataLayerChanged(EDataLayerAction::Modify, InDataLayerInstance, NAME_None);
@@ -377,6 +378,7 @@ void UDataLayerEditorSubsystem::AddToActorEditorContext(UDataLayerInstance* InDa
 
 void UDataLayerEditorSubsystem::RemoveFromActorEditorContext(UDataLayerInstance* InDataLayerInstance)
 {
+	check(InDataLayerInstance->CanBeInActorEditorContext());
 	if (InDataLayerInstance->RemoveFromActorEditorContext())
 	{
 		BroadcastDataLayerChanged(EDataLayerAction::Modify, InDataLayerInstance, NAME_None);
@@ -468,7 +470,7 @@ UWorld* UDataLayerEditorSubsystem::GetWorld() const
 
 bool UDataLayerEditorSubsystem::SetParentDataLayer(UDataLayerInstance* DataLayerInstance, UDataLayerInstance* ParentDataLayer)
 {
-	if (DataLayerInstance->CanParent(ParentDataLayer))
+	if (DataLayerInstance->CanBeChildOf(ParentDataLayer))
 	{
 		const bool bIsLoaded = DataLayerInstance->IsEffectiveLoadedInEditor();
 		DataLayerInstance->SetParent(ParentDataLayer);
@@ -1248,6 +1250,11 @@ void UDataLayerEditorSubsystem::DeleteDataLayers(const TArray<UDataLayerInstance
 	TArray<UDataLayerInstance*> DeletedDataLayerInstances;
 	for (UDataLayerInstance* DataLayerToDelete : DataLayersToDelete)
 	{
+		if (!DataLayerToDelete->IsUserManaged())
+		{
+			continue;
+		}
+
 		if (AWorldDataLayers* OuterWorldDataLayers = DataLayerToDelete ? DataLayerToDelete->GetOuterWorldDataLayers() : nullptr)
 		{
 			if (OuterWorldDataLayers->RemoveDataLayer(DataLayerToDelete))
@@ -1265,6 +1272,12 @@ void UDataLayerEditorSubsystem::DeleteDataLayers(const TArray<UDataLayerInstance
 void UDataLayerEditorSubsystem::DeleteDataLayer(UDataLayerInstance* DataLayerToDelete)
 {
 	UE_CLOG(!GetWorld(), LogDataLayerEditorSubsystem, Error, TEXT("%s - Failed because world in null."), ANSI_TO_TCHAR(__FUNCTION__));
+	
+	if (!DataLayerToDelete->IsUserManaged())
+	{
+		return;
+	}
+
 	if (AWorldDataLayers* OuterWorldDataLayers = DataLayerToDelete ? DataLayerToDelete->GetOuterWorldDataLayers() : nullptr)
 	{
 		if (OuterWorldDataLayers->RemoveDataLayer(DataLayerToDelete))
diff --git a/Engine/Source/Editor/DataLayerEditor/Private/DataLayer/DataLayerMode.cpp b/Engine/Source/Editor/DataLayerEditor/Private/DataLayer/DataLayerMode.cpp
--- a/Engine/Source/Editor/DataLayerEditor/Private/DataLayer/DataLayerMode.cpp
+++ b/Engine/Source/Editor/DataLayerEditor/Private/DataLayer/DataLayerMode.cpp
@@ -420,7 +420,7 @@ void FDataLayerMode::OnItemDoubleClick(FSceneOutlinerTreeItemPtr Item)
 	{
 		if (UDataLayerInstance* DataLayerInstance = DataLayerItem->GetDataLayer())
 		{
-			if (!DataLayerInstance->IsLocked())
+			if (DataLayerInstance->CanBeInActorEditorContext())
 			{
 				if (!DataLayerInstance->IsInActorEditorContext())
 				{
@@ -687,23 +687,10 @@ FSceneOutlinerDragValidationInfo FDataLayerMode::ValidateDrop(const ISceneOutlin
 
 			if (ParentDataLayer != nullptr)
 			{
-				if (!DataLayerInstance->CanParent(ParentDataLayer))
+				FText FailureReason;
+				if (!DataLayerInstance->CanBeChildOf(ParentDataLayer, &FailureReason))
 				{
-					if (!DataLayerInstance->IsDataLayerTypeValidToParent(ParentDataLayer->GetType()))
-					{
-						return FSceneOutlinerDragValidationInfo(ESceneOutlinerDropCompatibility::IncompatibleGeneric,
-							FText::Format(LOCTEXT("CantMoveToDataLayerDiffType", "Can't move a {0} Data Layer under a {1} Data Layer"), UEnum::GetDisplayValueAsText(DataLayerInstance->GetType()), UEnum::GetDisplayValueAsText(ParentDataLayer->GetType())));
-					}
-
-					if (ParentDataLayer == DataLayerInstance->GetParent() || ParentDataLayer == DataLayerInstance)
-					{
-						return FSceneOutlinerDragValidationInfo(ESceneOutlinerDropCompatibility::IncompatibleGeneric, LOCTEXT("CantMoveToSameDataLayer", "Can't move Data Layer to same Data Layer"));
-					}
-
-					if (ParentDataLayer->GetOuterWorldDataLayers() != DataLayerInstance->GetOuterWorldDataLayers())
-					{
-						return FSceneOutlinerDragValidationInfo(ESceneOutlinerDropCompatibility::IncompatibleGeneric, LOCTEXT("CantMoveToDifferentOuterWorldDataLayer", "Can't move Data Layer to a different World Data Layer"));
-					}
+					return FSceneOutlinerDragValidationInfo(ESceneOutlinerDropCompatibility::IncompatibleGeneric, FailureReason);
 				}
 
 				if (ParentDataLayer->IsLocked())
@@ -1058,7 +1045,7 @@ void FDataLayerMode::SetParentDataLayer(const TArray<UDataLayerInstance*> DataLa
 		ValidDataLayers.Reserve(DataLayerInstances.Num());
 		for (UDataLayerInstance* DataLayerInstance : DataLayerInstances)
 		{
-			if (DataLayerInstance->CanParent(ParentDataLayer))
+			if (DataLayerInstance->CanBeChildOf(ParentDataLayer))
 			{
 				ValidDataLayers.Add(DataLayerInstance);
 			}
@@ -1300,7 +1287,7 @@ void FDataLayerMode::RegisterContextMenu()
 			FDataLayerMode* Mode = const_cast<FDataLayerMode*>(static_cast<const FDataLayerMode*>(SceneOutliner->GetMode()));
 			check(Mode);
 			TArray<UDataLayerInstance*> SelectedDataLayers = Mode->GetSelectedDataLayers(SceneOutliner);
-			const bool bSelectedDataLayersContainsLocked = !!SelectedDataLayers.FindByPredicate([](const UDataLayerInstance* DataLayer) { return DataLayer->IsLocked(); });
+			const bool bSelectedDataLayersContainsLocked = !!SelectedDataLayers.FindByPredicate([](const UDataLayerInstance* DataLayerInstance) { return DataLayerInstance->IsLocked(); });
 
 			bool bHasActorEditorContextDataLayers = false;			
 			TArray<const UDataLayerInstance*> AllDataLayers;
@@ -1390,15 +1377,18 @@ void FDataLayerMode::RegisterContextMenu()
 				Section.AddMenuEntry("CreateNewDataLayer", LOCTEXT("CreateNewDataLayer", "Create New Data Layer"), FText(), FSlateIcon(),
 					FUIAction(FExecuteAction::CreateLambda([CreateNewEmptyDataLayer]() { CreateNewEmptyDataLayer(); })));
 
+				const bool bAllSelectedDataLayersCanAddActors = !!SelectedDataLayers.FindByPredicate([](const UDataLayerInstance* DataLayerInstance) { return !DataLayerInstance->CanUserAddActors(); });
 				Section.AddMenuEntry("AddSelectedActorsToSelectedDataLayers", LOCTEXT("AddSelectedActorsToSelectedDataLayersMenu", "Add Selected Actors to Selected Data Layers"), FText(), FSlateIcon(),
 					FUIAction(
-						FExecuteAction::CreateLambda([SelectedDataLayers]() {
+						FExecuteAction::CreateLambda([SelectedDataLayers]() 
+						{
 							check(!SelectedDataLayers.IsEmpty());
 							{
 								const FScopedTransaction Transaction(LOCTEXT("AddSelectedActorsToSelectedDataLayers", "Add Selected Actors to Selected Data Layers"));
 								UDataLayerEditorSubsystem::Get()->AddSelectedActorsToDataLayers(SelectedDataLayers);
-							}}),
-						FCanExecuteAction::CreateLambda([SelectedDataLayers,bSelectedDataLayersContainsLocked] { return !SelectedDataLayers.IsEmpty() && GEditor->GetSelectedActorCount() > 0 && !bSelectedDataLayersContainsLocked; })
+							}
+						}),
+						FCanExecuteAction::CreateLambda([SelectedDataLayers, bAllSelectedDataLayersCanAddActors] { return !SelectedDataLayers.IsEmpty() && GEditor->GetSelectedActorCount() > 0 && !bAllSelectedDataLayersCanAddActors; })
 					));
 
 				if (!Mode->SelectedDataLayerActors.IsEmpty())
@@ -1419,28 +1409,31 @@ void FDataLayerMode::RegisterContextMenu()
 							}));
 						}));
 				}
-				if (!SelectedDataLayers.IsEmpty() && !bSelectedDataLayersContainsLocked)
+
+				const bool bSelectedDataLayersSupportParenting = !SelectedDataLayers.FindByPredicate([](const UDataLayerInstance* DataLayerInstance) { return !DataLayerInstance->CanHaveParentDataLayer(); });
+				if (!SelectedDataLayers.IsEmpty() && !bSelectedDataLayersContainsLocked && bSelectedDataLayersSupportParenting)
 				{
 					Section.AddSubMenu("MoveSelectedDataLayersTo", LOCTEXT("MoveSelectedDataLayersTo", "Move Data Layers To"), FText(),
 						FNewToolMenuDelegate::CreateLambda([Mode](UToolMenu* InSubMenu)
 						{
 							CreateDataLayerPicker(InSubMenu, FOnDataLayerPicked::CreateLambda([Mode](UDataLayerInstance* TargetDataLayer)
 							{
-								TArray<UDataLayerInstance*> DataLayers;
-								for (TWeakObjectPtr<const UDataLayerInstance>& DataLayer : Mode->SelectedDataLayersSet)
+								TArray<UDataLayerInstance*> DataLayerInstances;
+								for (TWeakObjectPtr<const UDataLayerInstance>& DataLayerInstance : Mode->SelectedDataLayersSet)
 								{
-									if (DataLayer.IsValid() && !DataLayer->IsLocked() && DataLayer != TargetDataLayer)
+									if (DataLayerInstance.IsValid() && !DataLayerInstance->IsLocked() && DataLayerInstance->CanBeChildOf(TargetDataLayer))
 									{
-										DataLayers.Add(const_cast<UDataLayerInstance*>(DataLayer.Get()));
+										DataLayerInstances.Add(const_cast<UDataLayerInstance*>(DataLayerInstance.Get()));
 									}
 								}
-								Mode->SetParentDataLayer(DataLayers, TargetDataLayer);
+								Mode->SetParentDataLayer(DataLayerInstances, TargetDataLayer);
 							}), /*bShowRoot*/true);
 						}));
 				}
 
 				Section.AddSeparator("SectionsSeparator");
 
+				const bool bAllSelectedDataLayersCanRemoveActors = !!SelectedDataLayers.FindByPredicate([](const UDataLayerInstance* DataLayerInstance) { return !DataLayerInstance->CanUserRemoveActors(); });
 				Section.AddMenuEntry("RemoveSelectedActorsFromSelectedDataLayers", LOCTEXT("RemoveSelectedActorsFromSelectedDataLayersMenu", "Remove Selected Actors from Selected Data Layers"), FText(), FSlateIcon(),
 					FUIAction(
 						FExecuteAction::CreateLambda([SelectedDataLayers]() {
@@ -1449,17 +1442,21 @@ void FDataLayerMode::RegisterContextMenu()
 								const FScopedTransaction Transaction(LOCTEXT("RemoveSelectedActorsFromSelectedDataLayers_DataLayerMode", "Remove Selected Actors from Selected Data Layers"));
 								UDataLayerEditorSubsystem::Get()->RemoveSelectedActorsFromDataLayers(SelectedDataLayers);
 							}}),
-						FCanExecuteAction::CreateLambda([SelectedDataLayers,bSelectedDataLayersContainsLocked] { return !SelectedDataLayers.IsEmpty() && GEditor->GetSelectedActorCount() > 0 && !bSelectedDataLayersContainsLocked; })
+						FCanExecuteAction::CreateLambda([SelectedDataLayers, bAllSelectedDataLayersCanRemoveActors] { return !SelectedDataLayers.IsEmpty() && GEditor->GetSelectedActorCount() > 0 && !bAllSelectedDataLayersCanRemoveActors; })
 					));
 
+				const bool bAllSelectedDataLayersAreUserManaged = !SelectedDataLayers.FindByPredicate([](const UDataLayerInstance* DataLayerInstance) { return !DataLayerInstance->IsUserManaged(); });
 				Section.AddMenuEntry("DeleteSelectedDataLayers", LOCTEXT("DeleteSelectedDataLayers", "Delete Selected Data Layers"), FText(), FSlateIcon(),
 					FUIAction(
-						FExecuteAction::CreateLambda([Mode, SelectedDataLayers]() {
+						FExecuteAction::CreateLambda([Mode, SelectedDataLayers]() 
+						{
 							check(!SelectedDataLayers.IsEmpty());
-							{
-								Mode->DeleteDataLayers(SelectedDataLayers);
-							}}),
-						FCanExecuteAction::CreateLambda([SelectedDataLayers,bSelectedDataLayersContainsLocked] { return !SelectedDataLayers.IsEmpty() && !bSelectedDataLayersContainsLocked; })
+							Mode->DeleteDataLayers(SelectedDataLayers);
+						}),
+						FCanExecuteAction::CreateLambda([SelectedDataLayers, bSelectedDataLayersContainsLocked, bAllSelectedDataLayersAreUserManaged]
+						{
+							return !SelectedDataLayers.IsEmpty() && !bSelectedDataLayersContainsLocked && bAllSelectedDataLayersAreUserManaged;
+						})
 					));
 
 				if (UDataLayerEditorSubsystem::Get()->HasDeprecatedDataLayers())
@@ -1555,15 +1552,21 @@ void FDataLayerMode::RegisterContextMenu()
 						FExecuteAction::CreateLambda([Mode, SceneOutliner, SelectedDataLayers]() 
 						{
 							const FScopedTransaction Transaction(LOCTEXT("MakeCurrentDataLayers", "Make Current Data Layer(s)"));
-							for (TWeakObjectPtr<const UDataLayerInstance>& DataLayer : Mode->SelectedDataLayersSet)
+							for (TWeakObjectPtr<const UDataLayerInstance>& DataLayerInstance : Mode->SelectedDataLayersSet)
 							{
-								if (DataLayer.IsValid() && !DataLayer->IsLocked())
+								if (DataLayerInstance.IsValid())
 								{
-									UDataLayerEditorSubsystem::Get()->AddToActorEditorContext(const_cast<UDataLayerInstance*>(DataLayer.Get()));
+									UDataLayerEditorSubsystem::Get()->AddToActorEditorContext(const_cast<UDataLayerInstance*>(DataLayerInstance.Get()));
 								}
 							}
 						}),
-						FCanExecuteAction::CreateLambda([Mode] { return Algo::AnyOf(Mode->SelectedDataLayersSet, [](const TWeakObjectPtr<const UDataLayerInstance>& DataLayer) { return DataLayer.IsValid() && !DataLayer->IsLocked() && !DataLayer->IsInActorEditorContext(); }); })
+						FCanExecuteAction::CreateLambda([Mode] 
+						{ 
+							return Algo::AnyOf(Mode->SelectedDataLayersSet, [](const TWeakObjectPtr<const UDataLayerInstance>& DataLayerInstance)
+							{ 
+								return DataLayerInstance.IsValid() && DataLayerInstance->CanBeInActorEditorContext() && !DataLayerInstance->IsInActorEditorContext();
+							});
+						})
 					));
 
 				Section.AddMenuEntry("RemoveCurrentDataLayers", LOCTEXT("RemoveCurrentDataLayers", "Remove Current Data Layer(s)"), FText(), FSlateIcon(),
@@ -1571,15 +1574,21 @@ void FDataLayerMode::RegisterContextMenu()
 						FExecuteAction::CreateLambda([Mode, SceneOutliner, SelectedDataLayers]()
 						{
 							const FScopedTransaction Transaction(LOCTEXT("RemoveCurrentDataLayers", "Remove Current Data Layer(s)"));
-							for (TWeakObjectPtr<const UDataLayerInstance>& DataLayer : Mode->SelectedDataLayersSet)
+							for (TWeakObjectPtr<const UDataLayerInstance>& DataLayerInstance : Mode->SelectedDataLayersSet)
 							{
-								if (DataLayer.IsValid() && !DataLayer->IsLocked())
+								if (DataLayerInstance.IsValid())
 								{
-									UDataLayerEditorSubsystem::Get()->RemoveFromActorEditorContext(const_cast<UDataLayerInstance*>(DataLayer.Get()));
+									UDataLayerEditorSubsystem::Get()->RemoveFromActorEditorContext(const_cast<UDataLayerInstance*>(DataLayerInstance.Get()));
 								}
 							}
 						}),
-						FCanExecuteAction::CreateLambda([Mode] { return Algo::AnyOf(Mode->SelectedDataLayersSet, [](const TWeakObjectPtr<const UDataLayerInstance>& DataLayer) { return DataLayer.IsValid() && !DataLayer->IsLocked() && DataLayer->IsInActorEditorContext(); }); })
+						FCanExecuteAction::CreateLambda([Mode] 
+						{ 
+							return Algo::AnyOf(Mode->SelectedDataLayersSet, [](const TWeakObjectPtr<const UDataLayerInstance>& DatalLayerInstance) 
+							{ 
+								return DatalLayerInstance.IsValid() && DatalLayerInstance->CanBeInActorEditorContext() && DatalLayerInstance->IsInActorEditorContext();
+							});
+						})
 					));
 
 				Section.AddMenuEntry("ClearCurrentDataLayers", LOCTEXT("ClearCurrentDataLayers", "Clear Current Data Layers"), FText(), FSlateIcon(),
diff --git a/Engine/Source/Runtime/Engine/Private/WorldPartition/DataLayer/DataLayerInstance.cpp b/Engine/Source/Runtime/Engine/Private/WorldPartition/DataLayer/DataLayerInstance.cpp
--- a/Engine/Source/Runtime/Engine/Private/WorldPartition/DataLayer/DataLayerInstance.cpp
+++ b/Engine/Source/Runtime/Engine/Private/WorldPartition/DataLayer/DataLayerInstance.cpp
@@ -2,6 +2,7 @@
 
 #include "WorldPartition/DataLayer/DataLayerInstance.h"
 
+#include "Internationalization/Text.h"
 #include "UObject/UnrealType.h"
 #include "WorldPartition/WorldPartitionLog.h"
 #include "WorldPartition/DataLayer/DataLayerUtils.h"
@@ -134,6 +135,11 @@ const TCHAR* UDataLayerInstance::GetDataLayerIconName() const
 	return FDataLayerUtils::GetDataLayerIconName(GetType());
 }
 
+bool UDataLayerInstance::CanUserAddActors() const
+{
+	return !IsLocked();
+}
+
 bool UDataLayerInstance::CanAddActor(AActor* InActor) const
 {
 	return InActor != nullptr && InActor->SupportsDataLayer() && !InActor->ContainsDataLayer(this);
@@ -149,6 +155,11 @@ bool UDataLayerInstance::AddActor(AActor* InActor) const
 	return false;
 }
 
+bool UDataLayerInstance::CanUserRemoveActors() const
+{
+	return !IsLocked();
+}
+
 bool UDataLayerInstance::CanRemoveActor(AActor* InActor) const
 {
 	return InActor->GetDataLayerInstances().Contains(this) || InActor->GetDataLayerInstancesForLevel().Contains(this);
@@ -164,15 +175,64 @@ bool UDataLayerInstance::RemoveActor(AActor* InActor) const
 	return false;
 }
 
-bool UDataLayerInstance::CanParent(const UDataLayerInstance* InParent) const
+bool UDataLayerInstance::CanBeChildOf(const UDataLayerInstance* InParent, FText* OutReason) const
 {
-	return (this != InParent)
-		&& (Parent != InParent) 
-		&& (InParent == nullptr || (IsDataLayerTypeValidToParent(InParent->GetType()) && (InParent->GetOuterWorldDataLayers() == GetOuterWorldDataLayers())));
+	auto AssignReason = [OutReason](FText&& Reason)
+	{
+		if(OutReason != nullptr)
+		{
+			*OutReason = MoveTemp(Reason);
+		}
+	};
+
+	if (this == InParent || Parent == InParent)
+	{
+		AssignReason(LOCTEXT("SameParentOrSameDataLayer", "Data Layer already has this parent"));
+		return false;
+	}
+
+	if (InParent == nullptr)
+	{
+		// nullptr is considered a valid parent
+		return true;
+	}
+
+	if (!CanHaveParentDataLayer())
+	{
+		AssignReason(FText::Format(LOCTEXT("ParentDataLayerUnsupported", "Data Layer \"{0}\" does not support parent data layers"), FText::FromString(GetDataLayerShortName())));
+		return false;
+	}
+
+	if (!InParent->CanHaveChildDataLayers())
+	{
+		AssignReason(FText::Format(LOCTEXT("ChildDataLayerUnsuported", "Data Layer \"{0}\" does not support child data layers"), FText::FromString(InParent->GetDataLayerShortName())));
+		return false;
+	}
+
+	if (!IsParentDataLayerTypeCompatible(InParent))
+	{
+		AssignReason(FText::Format(LOCTEXT("IncompatibleChildType", "{0} Data Layer cannot have {1} child Data Layers"), UEnum::GetDisplayValueAsText(GetType()), UEnum::GetDisplayValueAsText(InParent->GetType())));
+		return false;
+	}
+
+	if (InParent->GetOuterWorldDataLayers() != GetOuterWorldDataLayers())
+	{
+		AssignReason(LOCTEXT("DifferentOuterWorldDataLayer", "Parent WorldDataLayers is a different from child WorldDataLayers"));
+		return false;
+	}
+
+	return true;
 }
 
-bool UDataLayerInstance::IsDataLayerTypeValidToParent(EDataLayerType ParentDataLayerType) const
+bool UDataLayerInstance::IsParentDataLayerTypeCompatible(const UDataLayerInstance* InParent) const
 {
+	if (InParent == nullptr)
+	{
+		return false;
+	}
+
+	EDataLayerType ParentDataLayerType = InParent->GetType();
+
 	return GetType() != EDataLayerType::Unknown
 		&& ParentDataLayerType != EDataLayerType::Unknown
 		&& (ParentDataLayerType == EDataLayerType::Editor || GetType() == EDataLayerType::Runtime);
@@ -180,11 +240,13 @@ bool UDataLayerInstance::IsDataLayerTypeValidToParent(EDataLayerType ParentDataL
 
 bool UDataLayerInstance::SetParent(UDataLayerInstance* InParent)
 {
-	if (!CanParent(InParent))
+	if (!CanBeChildOf(InParent))
 	{
 		return false;
 	}
 
+	check(CanHaveParentDataLayer());
+
 	Modify();
 
 	// If we find ourself in the parent chain of the provided parent
@@ -220,6 +282,8 @@ void UDataLayerInstance::SetChildParent(UDataLayerInstance* InParent)
 		return;
 	}
 
+	check(!InParent || InParent->CanHaveChildDataLayers());
+
 	Modify();
 	while (Children.Num())
 	{
@@ -241,7 +305,7 @@ FText UDataLayerInstance::GetDataLayerText(const UDataLayerInstance* InDataLayer
 
 bool UDataLayerInstance::Validate(IStreamingGenerationErrorHandler* ErrorHandler) const
 {
-	if (GetParent() != nullptr && !IsDataLayerTypeValidToParent(GetParent()->GetType()))
+	if (GetParent() != nullptr && !IsParentDataLayerTypeCompatible(GetParent()))
 	{
 		ErrorHandler->OnDataLayerHierarchyTypeMismatch(this, GetParent());
 		return false;
@@ -250,6 +314,11 @@ bool UDataLayerInstance::Validate(IStreamingGenerationErrorHandler* ErrorHandler
 	return true;
 }
 
+bool UDataLayerInstance::CanBeInActorEditorContext() const
+{
+	return !IsLocked();
+}
+
 bool UDataLayerInstance::IsInActorEditorContext() const
 {
 	return GetOuterWorldDataLayers()->IsInActorEditorContext(this);
@@ -315,6 +384,7 @@ void UDataLayerInstance::ForEachChild(TFunctionRef<bool(const UDataLayerInstance
 void UDataLayerInstance::AddChild(UDataLayerInstance* InDataLayer)
 {
 	check(InDataLayer->GetOuterWorldDataLayers() == GetOuterWorldDataLayers())
+	check(CanHaveChildDataLayers());
 	Modify();
 	checkSlow(!Children.Contains(InDataLayer));
 	Children.Add(InDataLayer);
diff --git a/Engine/Source/Runtime/Engine/Private/WorldPartition/DataLayer/WorldDataLayers.cpp b/Engine/Source/Runtime/Engine/Private/WorldPartition/DataLayer/WorldDataLayers.cpp
--- a/Engine/Source/Runtime/Engine/Private/WorldPartition/DataLayer/WorldDataLayers.cpp
+++ b/Engine/Source/Runtime/Engine/Private/WorldPartition/DataLayer/WorldDataLayers.cpp
@@ -575,7 +575,9 @@ bool AWorldDataLayers::IsInActorEditorContext(const UDataLayerInstance* InDataLa
 
 bool AWorldDataLayers::AddToActorEditorContext(UDataLayerInstance* InDataLayerInstance)
 {
-	check(DataLayerInstances.Contains(InDataLayerInstance));
+	check(InDataLayerInstance->CanBeInActorEditorContext());
+	check(ContainsDataLayer(InDataLayerInstance));
+
 	if (!CurrentDataLayers.DataLayerInstanceNames.Contains(InDataLayerInstance->GetDataLayerFName()))
 	{
 		Modify(/*bDirty*/false);
@@ -587,7 +589,9 @@ bool AWorldDataLayers::AddToActorEditorContext(UDataLayerInstance* InDataLayerIn
 
 bool AWorldDataLayers::RemoveFromActorEditorContext(UDataLayerInstance* InDataLayerInstance)
 {
-	check(DataLayerInstances.Contains(InDataLayerInstance));
+	check(InDataLayerInstance->CanBeInActorEditorContext());
+	check(ContainsDataLayer(InDataLayerInstance));
+
 	if (CurrentDataLayers.DataLayerInstanceNames.Contains(InDataLayerInstance->GetDataLayerFName()))
 	{
 		Modify(/*bDirty*/false);
diff --git a/Engine/Source/Runtime/Engine/Public/WorldPartition/DataLayer/DataLayerInstance.h b/Engine/Source/Runtime/Engine/Public/WorldPartition/DataLayer/DataLayerInstance.h
--- a/Engine/Source/Runtime/Engine/Public/WorldPartition/DataLayer/DataLayerInstance.h
+++ b/Engine/Source/Runtime/Engine/Public/WorldPartition/DataLayer/DataLayerInstance.h
@@ -13,6 +13,7 @@
 #define DATALAYER_TO_INSTANCE_RUNTIME_CONVERSION_ENABLED 1
 
 class AWorldDataLayers;
+class FText;
 class IStreamingGenerationErrorHandler;
 class UWorld;
 
@@ -78,8 +79,9 @@ public:
 
 	const TCHAR* GetDataLayerIconName() const;
 
-	bool CanParent(const UDataLayerInstance* InParent) const;
-	bool IsDataLayerTypeValidToParent(EDataLayerType ParentDataLayerType) const;
+	virtual bool CanHaveChildDataLayers() const { return true; }
+	virtual bool CanHaveParentDataLayer() const { return true; }
+	bool CanBeChildOf(const UDataLayerInstance* InParent, FText* OutReason = nullptr) const;
 	bool SetParent(UDataLayerInstance* InParent);
 
 	void SetChildParent(UDataLayerInstance* InParent);
@@ -90,15 +92,23 @@ public:
 	virtual bool IsReadOnly() const;
 	virtual bool CanEditChange(const FProperty* InProperty) const;
 	
+	virtual bool CanUserAddActors() const;
 	virtual bool CanAddActor(AActor* Actor) const;
 	virtual bool AddActor(AActor* Actor) const;
+	virtual bool CanUserRemoveActors() const;
 	virtual bool CanRemoveActor(AActor* Actor) const;
 	virtual bool RemoveActor(AActor* Actor) const;
 
+	virtual bool CanBeInActorEditorContext() const;
 	bool IsInActorEditorContext() const;
 	bool AddToActorEditorContext();
 	bool RemoveFromActorEditorContext();
 
+	virtual bool CanEditDataLayerShortName() const { return false; }
+
+	// Whether the DataLayer was created by a user and can be deleted by a user.
+	virtual bool IsUserManaged() const { return true; }
+
 	virtual bool Validate(IStreamingGenerationErrorHandler* ErrorHandler) const;
 
 	UE_DEPRECATED(5.3, "Use CanEditShortName instead")
@@ -107,7 +117,8 @@ public:
 	UE_DEPRECATED(5.3, "Use SetShortName instead")
 	virtual bool RelabelDataLayer(FName NewDataLayerLabel) { return false; }
 
-	virtual bool CanEditDataLayerShortName() const { return false; }
+	UE_DEPRECATED(5.3, "Use CanBeChildOf instead")
+	bool CanParent(const UDataLayerInstance* InParent) const { return CanBeChildOf(InParent); }
 #endif // WITH_EDITOR
 
 	UFUNCTION(Category = "Data Layer", BlueprintCallable)
@@ -164,6 +175,8 @@ private:
 
 protected:
 #if WITH_EDITOR
+	bool IsParentDataLayerTypeCompatible(const UDataLayerInstance* InParent) const;
+
 	virtual bool PerformAddActor(AActor* InActor) const { return false; }
 	virtual bool PerformRemoveActor(AActor* InActor) const { return false;  }
 	virtual void PerformSetDataLayerShortName(const FString& InNewShortName) {}
