Fix static analysis error

diff --git a/Engine/Source/Editor/DataLayerEditor/Private/DataLayer/DataLayerEditorSubsystem.cpp b/Engine/Source/Editor/DataLayerEditor/Private/DataLayer/DataLayerEditorSubsystem.cpp
--- a/Engine/Source/Editor/DataLayerEditor/Private/DataLayer/DataLayerEditorSubsystem.cpp
+++ b/Engine/Source/Editor/DataLayerEditor/Private/DataLayer/DataLayerEditorSubsystem.cpp
@@ -1259,17 +1259,20 @@ void UDataLayerEditorSubsystem::DeleteDataLayers(const TArray<UDataLayerInstance
 	TArray<UDataLayerInstance*> DeletedDataLayerInstances;
 	for (UDataLayerInstance* DataLayerToDelete : DataLayersToDelete)
 	{
+		if (!DataLayerToDelete)
+		{
+			continue;
+		}
+
 		if (!DataLayerToDelete->IsUserManaged())
 		{
 			continue;
 		}
 
-		if (AWorldDataLayers* OuterWorldDataLayers = DataLayerToDelete ? DataLayerToDelete->GetOuterWorldDataLayers() : nullptr)
+		AWorldDataLayers* OuterWorldDataLayers = DataLayerToDelete->GetOuterWorldDataLayers();
+		if (OuterWorldDataLayers->RemoveDataLayer(DataLayerToDelete))
 		{
-			if (OuterWorldDataLayers->RemoveDataLayer(DataLayerToDelete))
-			{
-				DeletedDataLayerInstances.Add(DataLayerToDelete);
-			}
+			DeletedDataLayerInstances.Add(DataLayerToDelete);
 		}
 	}
 	for (UDataLayerInstance* DeletedDataLayerInstance : DeletedDataLayerInstances)
@@ -1281,18 +1284,21 @@ void UDataLayerEditorSubsystem::DeleteDataLayers(const TArray<UDataLayerInstance
 void UDataLayerEditorSubsystem::DeleteDataLayer(UDataLayerInstance* DataLayerToDelete)
 {
 	UE_CLOG(!GetWorld(), LogDataLayerEditorSubsystem, Error, TEXT("%s - Failed because world in null."), ANSI_TO_TCHAR(__FUNCTION__));
+
+	if (!DataLayerToDelete)
+	{
+		return;
+	}
 	
 	if (!DataLayerToDelete->IsUserManaged())
 	{
 		return;
 	}
 
-	if (AWorldDataLayers* OuterWorldDataLayers = DataLayerToDelete ? DataLayerToDelete->GetOuterWorldDataLayers() : nullptr)
+	AWorldDataLayers* OuterWorldDataLayers = DataLayerToDelete->GetOuterWorldDataLayers();
+	if (OuterWorldDataLayers->RemoveDataLayer(DataLayerToDelete))
 	{
-		if (OuterWorldDataLayers->RemoveDataLayer(DataLayerToDelete))
-		{
-			BroadcastDataLayerChanged(EDataLayerAction::Delete, DataLayerToDelete, NAME_None);
-		}
+		BroadcastDataLayerChanged(EDataLayerAction::Delete, DataLayerToDelete, NAME_None);
 	}
 }
 
