Non-user managed actors can no longer be added/removed to DataLayers via the DataLayerOutliner (similar to DragnDropOp)
Removed DataLayerMode capture of this as it is not valid when the context menu is re-created

diff --git a/Engine/Source/Editor/DataLayerEditor/Private/DataLayer/DataLayerMode.cpp b/Engine/Source/Editor/DataLayerEditor/Private/DataLayer/DataLayerMode.cpp
--- a/Engine/Source/Editor/DataLayerEditor/Private/DataLayer/DataLayerMode.cpp
+++ b/Engine/Source/Editor/DataLayerEditor/Private/DataLayer/DataLayerMode.cpp
@@ -5,6 +5,7 @@
 #include "ActorDescTreeItem.h"
 #include "ActorMode.h"
 #include "Algo/AnyOf.h"
+#include "Algo/AllOf.h"
 #include "Algo/Transform.h"
 #include "Algo/Accumulate.h"
 #include "AssetRegistry/IAssetRegistry.h"
@@ -1287,7 +1288,12 @@ void FDataLayerMode::RegisterContextMenu()
 			FDataLayerMode* Mode = const_cast<FDataLayerMode*>(static_cast<const FDataLayerMode*>(SceneOutliner->GetMode()));
 			check(Mode);
 			TArray<UDataLayerInstance*> SelectedDataLayers = Mode->GetSelectedDataLayers(SceneOutliner);
-			const bool bSelectedDataLayersContainsLocked = !!SelectedDataLayers.FindByPredicate([](const UDataLayerInstance* DataLayerInstance) { return DataLayerInstance->IsLocked(); });
+			const bool bSelectedDataLayersContainsLocked = Algo::AnyOf(SelectedDataLayers, [](const UDataLayerInstance* DataLayerInstance) { return DataLayerInstance->IsLocked(); });
+
+			TArray<AActor*> SelectedActors;
+			GEditor->GetSelectedActors()->GetSelectedObjects<AActor>(SelectedActors);
+			const bool bSelectedActorsAreAllUserManaged = Algo::AllOf(SelectedActors, [](const AActor* Actor) { return Actor->IsUserManaged(); });
+			static const FText SelectionContainsNonUserManagedActorsText = LOCTEXT("SelectionContainsUserManagedActors", "Selection Contains Non User Managed Actors");
 
 			bool bHasActorEditorContextDataLayers = false;			
 			TArray<const UDataLayerInstance*> AllDataLayers;
@@ -1385,8 +1391,9 @@ void FDataLayerMode::RegisterContextMenu()
 						FIsActionChecked(),
 						FIsActionButtonVisible::CreateLambda([]() { return !UDataLayerEditorSubsystem::Get()->HasDeprecatedDataLayers(); })));
 
-				const bool bAllSelectedDataLayersCanAddActors = !!SelectedDataLayers.FindByPredicate([](const UDataLayerInstance* DataLayerInstance) { return !DataLayerInstance->CanUserAddActors(); });
-				Section.AddMenuEntry("AddSelectedActorsToSelectedDataLayers", LOCTEXT("AddSelectedActorsToSelectedDataLayersMenu", "Add Selected Actors to Selected Data Layers"), FText(), FSlateIcon(),
+				const bool bAllSelectedDataLayersCanAddActors = Algo::AllOf(SelectedDataLayers, [](const UDataLayerInstance* DataLayerInstance) { return DataLayerInstance->CanUserAddActors(); });
+				FText AddActorToolTip = bSelectedActorsAreAllUserManaged ? FText::GetEmpty() : SelectionContainsNonUserManagedActorsText;
+				Section.AddMenuEntry("AddSelectedActorsToSelectedDataLayers", LOCTEXT("AddSelectedActorsToSelectedDataLayersMenu", "Add Selected Actors to Selected Data Layers"), AddActorToolTip, FSlateIcon(),
 					FUIAction(
 						FExecuteAction::CreateLambda([SelectedDataLayers]() 
 						{
@@ -1396,7 +1403,10 @@ void FDataLayerMode::RegisterContextMenu()
 								UDataLayerEditorSubsystem::Get()->AddSelectedActorsToDataLayers(SelectedDataLayers);
 							}
 						}),
-						FCanExecuteAction::CreateLambda([SelectedDataLayers, bAllSelectedDataLayersCanAddActors] { return !SelectedDataLayers.IsEmpty() && GEditor->GetSelectedActorCount() > 0 && !bAllSelectedDataLayersCanAddActors; })
+						FCanExecuteAction::CreateLambda([SelectedDataLayers, bAllSelectedDataLayersCanAddActors, bSelectedActorsAreAllUserManaged]
+						{ 
+							return !SelectedDataLayers.IsEmpty() && GEditor->GetSelectedActorCount() > 0 && bAllSelectedDataLayersCanAddActors && bSelectedActorsAreAllUserManaged;
+						})
 					));
 
 				if (!Mode->SelectedDataLayerActors.IsEmpty())
@@ -1418,7 +1428,7 @@ void FDataLayerMode::RegisterContextMenu()
 						}));
 				}
 
-				const bool bSelectedDataLayersSupportParenting = !SelectedDataLayers.FindByPredicate([](const UDataLayerInstance* DataLayerInstance) { return !DataLayerInstance->CanHaveParentDataLayer(); });
+				const bool bSelectedDataLayersSupportParenting = Algo::AllOf(SelectedDataLayers, [](const UDataLayerInstance* DataLayerInstance) { return DataLayerInstance->CanHaveParentDataLayer(); });
 				if (!SelectedDataLayers.IsEmpty() && !bSelectedDataLayersContainsLocked && bSelectedDataLayersSupportParenting)
 				{
 					Section.AddSubMenu("MoveSelectedDataLayersTo", LOCTEXT("MoveSelectedDataLayersTo", "Move Data Layers To"), FText(),
@@ -1441,19 +1451,25 @@ void FDataLayerMode::RegisterContextMenu()
 
 				Section.AddSeparator("SectionsSeparator");
 
-				const bool bAllSelectedDataLayersCanRemoveActors = !!SelectedDataLayers.FindByPredicate([](const UDataLayerInstance* DataLayerInstance) { return !DataLayerInstance->CanUserRemoveActors(); });
-				Section.AddMenuEntry("RemoveSelectedActorsFromSelectedDataLayers", LOCTEXT("RemoveSelectedActorsFromSelectedDataLayersMenu", "Remove Selected Actors from Selected Data Layers"), FText(), FSlateIcon(),
+				const bool bAllSelectedDataLayersCanRemoveActors = Algo::AllOf(SelectedDataLayers, [](const UDataLayerInstance* DataLayerInstance) { return DataLayerInstance->CanUserRemoveActors(); });
+				FText RemoveActorToolTip = bSelectedActorsAreAllUserManaged ? FText::GetEmpty() : SelectionContainsNonUserManagedActorsText;
+				Section.AddMenuEntry("RemoveSelectedActorsFromSelectedDataLayers", LOCTEXT("RemoveSelectedActorsFromSelectedDataLayersMenu", "Remove Selected Actors from Selected Data Layers"), RemoveActorToolTip, FSlateIcon(),
 					FUIAction(
-						FExecuteAction::CreateLambda([SelectedDataLayers]() {
+						FExecuteAction::CreateLambda([SelectedDataLayers]() 
+						{
 							check(!SelectedDataLayers.IsEmpty());
 							{
 								const FScopedTransaction Transaction(LOCTEXT("RemoveSelectedActorsFromSelectedDataLayers_DataLayerMode", "Remove Selected Actors from Selected Data Layers"));
 								UDataLayerEditorSubsystem::Get()->RemoveSelectedActorsFromDataLayers(SelectedDataLayers);
-							}}),
-						FCanExecuteAction::CreateLambda([SelectedDataLayers, bAllSelectedDataLayersCanRemoveActors] { return !SelectedDataLayers.IsEmpty() && GEditor->GetSelectedActorCount() > 0 && !bAllSelectedDataLayersCanRemoveActors; })
+							}
+						}),
+						FCanExecuteAction::CreateLambda([SelectedDataLayers, bAllSelectedDataLayersCanRemoveActors, bSelectedActorsAreAllUserManaged]
+						{ 
+							return !SelectedDataLayers.IsEmpty() && GEditor->GetSelectedActorCount() > 0 && bAllSelectedDataLayersCanRemoveActors && bSelectedActorsAreAllUserManaged;
+						})
 					));
 
-				const bool bAllSelectedDataLayersAreUserManaged = !SelectedDataLayers.FindByPredicate([](const UDataLayerInstance* DataLayerInstance) { return !DataLayerInstance->IsUserManaged(); });
+				const bool bAllSelectedDataLayersAreUserManaged = Algo::AllOf(SelectedDataLayers, [](const UDataLayerInstance* DataLayerInstance) { return DataLayerInstance->IsUserManaged(); });
 				Section.AddMenuEntry("DeleteSelectedDataLayers", LOCTEXT("DeleteSelectedDataLayers", "Delete Selected Data Layers"), FText(), FSlateIcon(),
 					FUIAction(
 						FExecuteAction::CreateLambda([Mode, SelectedDataLayers]() 
