Prevent assigning actor to private data layer not part of the same WorldDataLayers

diff --git a/Engine/Source/Editor/DataLayerEditor/Private/DataLayer/DataLayerMode.cpp b/Engine/Source/Editor/DataLayerEditor/Private/DataLayer/DataLayerMode.cpp
--- a/Engine/Source/Editor/DataLayerEditor/Private/DataLayer/DataLayerMode.cpp
+++ b/Engine/Source/Editor/DataLayerEditor/Private/DataLayer/DataLayerMode.cpp
@@ -600,12 +600,21 @@ FSceneOutlinerDragValidationInfo FDataLayerMode::ValidateDrop(const ISceneOutlin
 	TArray<AActor*> PayloadActors = GetActorsFromOperation(Payload.SourceOperation);
 	if (!PayloadActors.IsEmpty()) // Adding actor(s) in data layer(s)
 	{
+		const FDataLayerTreeItem* DataLayerItem = DropTarget.CastTo<FDataLayerTreeItem>();
+		const UDataLayerInstance* TargetDataLayer = DataLayerItem ? DataLayerItem->GetDataLayer() : nullptr;
+		const bool bIsPrivateDataLayer = TargetDataLayer && TargetDataLayer->GetAsset() && TargetDataLayer->GetAsset()->IsPrivate();
+		const AWorldDataLayers* OuterWorldDataLayers = TargetDataLayer ? TargetDataLayer->GetOuterWorldDataLayers() : nullptr;
+
 		for (AActor* Actor : PayloadActors)
 		{
-			if (!Actor->IsUserManaged() ||  !DataLayerEditorSubsystem->IsActorValidForDataLayer(Actor))
+			if (!Actor->IsUserManaged() || !DataLayerEditorSubsystem->IsActorValidForDataLayer(Actor))
 			{
 				return FSceneOutlinerDragValidationInfo(ESceneOutlinerDropCompatibility::IncompatibleGeneric, LOCTEXT("ActorCantBeAssignedToDataLayer", "Can't assign actors to Data Layer"));
 			}
+			else if (bIsPrivateDataLayer && (Actor->GetLevel()->GetWorldDataLayers() != OuterWorldDataLayers))
+			{
+				return FSceneOutlinerDragValidationInfo(ESceneOutlinerDropCompatibility::IncompatibleGeneric, LOCTEXT("ActorCantBeAssignedToPrivateDataLayerOfOtherWorldDataLayers", "Can't assign actors to Private Data Layer of another WorldDataLayers actor"));
+			}
 		}
 
 		return ValidateActorDrop(DropTarget);
