RemoveAllDataLayers now goes through DataLayerInstance::RemoveActor interface to perform Removal validation code

diff --git a/Engine/Source/Editor/DataLayerEditor/Private/DataLayer/DataLayerEditorSubsystem.cpp b/Engine/Source/Editor/DataLayerEditor/Private/DataLayer/DataLayerEditorSubsystem.cpp
--- a/Engine/Source/Editor/DataLayerEditor/Private/DataLayer/DataLayerEditorSubsystem.cpp
+++ b/Engine/Source/Editor/DataLayerEditor/Private/DataLayer/DataLayerEditorSubsystem.cpp
@@ -583,13 +583,13 @@ bool UDataLayerEditorSubsystem::RemoveActorsFromAllDataLayers(const TArray<AActo
 {
 	GEditor->GetSelectedActors()->BeginBatchSelectOperation();
 
-	bool bChangesOccurred = false;
+	bool bRemoveAllDataLayersOnAllActor = true;
 	for (AActor* Actor : Actors)
 	{
-		TArray<const UDataLayerInstance*> ModifiedDataLayerInstances = Actor->GetDataLayerInstances();
-		if (Actor->RemoveAllDataLayers())
+		TArray<const UDataLayerInstance*> RemovedDataLayers = Actor->RemoveAllDataLayers();
+		if (!RemovedDataLayers.IsEmpty())
 		{
-			for (const UDataLayerInstance* DataLayerInstance : ModifiedDataLayerInstances)
+			for (const UDataLayerInstance* DataLayerInstance : RemovedDataLayers)
 			{
 				BroadcastDataLayerChanged(EDataLayerAction::Modify, DataLayerInstance, NAME_None);
 			}
@@ -600,13 +600,13 @@ bool UDataLayerEditorSubsystem::RemoveActorsFromAllDataLayers(const TArray<AActo
 			bool bActorSelectionChanged = false;
 			UpdateActorVisibility(Actor, bActorSelectionChanged, bActorModified, /*bActorNotifySelectionChange*/true, /*bActorRedrawViewports*/false);
 
-			bChangesOccurred = true;
+			bRemoveAllDataLayersOnAllActor &= !Actor->HasDataLayers();
 		}
 	}
 
 	GEditor->GetSelectedActors()->EndBatchSelectOperation();
 
-	return bChangesOccurred;
+	return bRemoveAllDataLayersOnAllActor;
 }
 
 bool UDataLayerEditorSubsystem::RemoveActorFromDataLayer(AActor* Actor, UDataLayerInstance* DataLayerInstance)
diff --git a/Engine/Source/Runtime/Engine/Classes/GameFramework/Actor.h b/Engine/Source/Runtime/Engine/Classes/GameFramework/Actor.h
--- a/Engine/Source/Runtime/Engine/Classes/GameFramework/Actor.h
+++ b/Engine/Source/Runtime/Engine/Classes/GameFramework/Actor.h
@@ -1298,7 +1298,7 @@ private:
 public:
 	bool AddDataLayer(const UDataLayerInstance* DataLayerInstance);
 	bool RemoveDataLayer(const UDataLayerInstance* DataLayerInstance);
-	bool RemoveAllDataLayers();
+	TArray<const UDataLayerInstance*> RemoveAllDataLayers();
 	bool SupportsDataLayer() const;
 	
 	TArray<const UDataLayerInstance*> GetDataLayerInstancesForLevel() const;
diff --git a/Engine/Source/Runtime/Engine/Private/ActorEditor.cpp b/Engine/Source/Runtime/Engine/Private/ActorEditor.cpp
--- a/Engine/Source/Runtime/Engine/Private/ActorEditor.cpp
+++ b/Engine/Source/Runtime/Engine/Private/ActorEditor.cpp
@@ -1641,16 +1641,17 @@ bool AActor::RemoveDataLayer(const UDataLayerInstance* DataLayerInstance)
 	return DataLayerInstance->RemoveActor(this);
 }
 
-bool AActor::RemoveAllDataLayers()
+TArray<const UDataLayerInstance*> AActor::RemoveAllDataLayers()
 {
-	if ((DataLayerAssets.Num() > 0) || (DataLayers.Num() > 0))
+	TArray<const UDataLayerInstance*> RemovedDataLayerInstances = GetDataLayerInstances();
+	if (!RemovedDataLayerInstances.IsEmpty())
 	{
 		Modify();
-		DataLayerAssets.Empty();
-		DataLayers.Empty();
-		return true;
+		RemovedDataLayerInstances.SetNum(Algo::RemoveIf(
+			RemovedDataLayerInstances, [this] (const UDataLayerInstance* DataLayerInstance) { return !DataLayerInstance->RemoveActor(this); }));
 	}
-	return false;
+
+	return RemovedDataLayerInstances;
 }
 
 TArray<const UDataLayerAsset*> AActor::ResolveDataLayerAssets(const TArray<TSoftObjectPtr<UDataLayerAsset>>& InDataLayerAssets) const
