DataLayer Outliner: Fix deleted actors showing as unloaded

#rnx

diff --git a/Engine/Source/Editor/DataLayerEditor/Private/DataLayer/DataLayerHierarchy.cpp b/Engine/Source/Editor/DataLayerEditor/Private/DataLayer/DataLayerHierarchy.cpp
--- a/Engine/Source/Editor/DataLayerEditor/Private/DataLayer/DataLayerHierarchy.cpp
+++ b/Engine/Source/Editor/DataLayerEditor/Private/DataLayer/DataLayerHierarchy.cpp
@@ -285,7 +285,7 @@ void FDataLayerHierarchy::CreateItems(TArray<FSceneOutlinerTreeItemPtr>& OutItem
 						// Create an FDataLayerActorDescTreeItem for each unloaded actor of this WorldPartition
 						FWorldPartitionHelpers::ForEachActorDesc(WorldPartition, [this, IsDataLayerShown, DataLayerManager, CurrentLevel, &WorldToLevelDataLayerMap, &OutItems](const FWorldPartitionActorDesc* ActorDesc)
 						{
-							if (ActorDesc != nullptr && !ActorDesc->IsLoaded() && FActorDescTreeItem::ShouldDisplayInOutliner(ActorDesc))
+							if (ActorDesc != nullptr && !ActorDesc->IsLoaded(true) && FActorDescTreeItem::ShouldDisplayInOutliner(ActorDesc))
 							{
 								for (const FName& DataLayerInstanceName : ActorDesc->GetDataLayerInstanceNames())
 								{
