DataLayerEditorSubsystem: Avoid SceneOutliner refreshes
- remove HasDataLayers() condition which was added for fixing CopyPaste of actors and no longer needed (tested copy/paste)
- avoid calling ActorEditorContextClientChanged in BroadcastDataLayerChanged and instead call it in AddActorToEditorContext/RemoveActorFromEditorContext

#rnx

diff --git a/Engine/Source/Editor/DataLayerEditor/Private/DataLayer/DataLayerEditorSubsystem.cpp b/Engine/Source/Editor/DataLayerEditor/Private/DataLayer/DataLayerEditorSubsystem.cpp
--- a/Engine/Source/Editor/DataLayerEditor/Private/DataLayer/DataLayerEditorSubsystem.cpp
+++ b/Engine/Source/Editor/DataLayerEditor/Private/DataLayer/DataLayerEditorSubsystem.cpp
@@ -151,7 +151,7 @@ void FDataLayersBroadcast::OnObjectPostEditChange(UObject* Object, FPropertyChan
 		}
 		else if (AActor* Actor = Cast<AActor>(Object))
 		{
-			bRefresh = Actor->IsPropertyChangedAffectingDataLayers(PropertyChangedEvent) || Actor->HasDataLayers();
+			bRefresh = Actor->IsPropertyChangedAffectingDataLayers(PropertyChangedEvent);
 		}
 		if (bRefresh)
 		{
@@ -384,6 +384,7 @@ void UDataLayerEditorSubsystem::AddToActorEditorContext(UDataLayerInstance* InDa
 	if (InDataLayerInstance->AddToActorEditorContext())
 	{
 		BroadcastDataLayerChanged(EDataLayerAction::Modify, InDataLayerInstance, NAME_None);
+		ActorEditorContextClientChanged.Broadcast(this);
 	}
 }
 
@@ -393,6 +394,7 @@ void UDataLayerEditorSubsystem::RemoveFromActorEditorContext(UDataLayerInstance*
 	if (InDataLayerInstance->RemoveFromActorEditorContext())
 	{
 		BroadcastDataLayerChanged(EDataLayerAction::Modify, InDataLayerInstance, NAME_None);
+		ActorEditorContextClientChanged.Broadcast(this);
 	}
 }
 
@@ -1321,7 +1323,6 @@ void UDataLayerEditorSubsystem::BroadcastDataLayerChanged(const EDataLayerAction
 {
 	bRebuildSelectedDataLayersFromEditorSelection = true;
 	DataLayerChanged.Broadcast(Action, ChangedDataLayer, ChangedProperty);
-	ActorEditorContextClientChanged.Broadcast(this);
 }
 
 void UDataLayerEditorSubsystem::OnDataLayerEditorLoadingStateChanged(bool bIsFromUserChange)
