SceneOutliner: Optimize cost of FWorldPartitionActorDesc::IsLoaded call (which calls StaticFindObject). Use new FEditorLoadedActorCache helper instead.

12 calls to FActorHierarchy::CreateWorldChildren in a 55k actor map with 15k loaded actors goes from 5.1 seconds to 2.7 seconds

#rnx

diff --git a/Engine/Source/Editor/DataLayerEditor/Private/DataLayer/DataLayerHierarchy.cpp b/Engine/Source/Editor/DataLayerEditor/Private/DataLayer/DataLayerHierarchy.cpp
--- a/Engine/Source/Editor/DataLayerEditor/Private/DataLayer/DataLayerHierarchy.cpp
+++ b/Engine/Source/Editor/DataLayerEditor/Private/DataLayer/DataLayerHierarchy.cpp
@@ -32,6 +32,7 @@
 #include "WorldPartition/WorldPartitionActorDesc.h"
 #include "WorldPartition/WorldPartitionHelpers.h"
 #include "WorldPartition/WorldPartitionSubsystem.h"
+#include "EditorLoadedActorCache.h"
 
 TUniquePtr<FDataLayerHierarchy> FDataLayerHierarchy::Create(FDataLayerMode* Mode, const TWeakObjectPtr<UWorld>& World)
 {
@@ -269,8 +270,9 @@ void FDataLayerHierarchy::CreateItems(TArray<FSceneOutlinerTreeItemPtr>& OutItem
 				const UDataLayerManager* DataLayerManager = GetOwningWorld()->GetDataLayerManager();
 				if (WorldPartitionSubsystem && DataLayerManager)
 				{
+					FEditorLoadedActorCache LoadedActorCache;
 					const ULevelInstanceSubsystem* LevelInstanceSubsystem = UWorld::GetSubsystem<ULevelInstanceSubsystem>(GetOwningWorld());
-					WorldPartitionSubsystem->ForEachWorldPartition([this, DataLayerManager, CurrentLevel, LevelInstanceSubsystem, IsDataLayerShown, &WorldToLevelDataLayerMap, &OutItems](UWorldPartition* WorldPartition)
+					WorldPartitionSubsystem->ForEachWorldPartition([this, DataLayerManager, CurrentLevel, LevelInstanceSubsystem, IsDataLayerShown, &WorldToLevelDataLayerMap, &LoadedActorCache, &OutItems](UWorldPartition* WorldPartition)
 					{
 						// Skip WorldPartition if it's not the one of the current level (the editing level instance)
 						// or if we hide the content of level instances
@@ -282,10 +284,12 @@ void FDataLayerHierarchy::CreateItems(TArray<FSceneOutlinerTreeItemPtr>& OutItem
 							return true;
 						}
 
+						const TSet<FGuid>& LoadedActors = LoadedActorCache.GetLoadedActorsForLevel(OuterLevel);
+
 						// Create an FDataLayerActorDescTreeItem for each unloaded actor of this WorldPartition
-						FWorldPartitionHelpers::ForEachActorDesc(WorldPartition, [this, IsDataLayerShown, DataLayerManager, CurrentLevel, &WorldToLevelDataLayerMap, &OutItems](const FWorldPartitionActorDesc* ActorDesc)
+						FWorldPartitionHelpers::ForEachActorDesc(WorldPartition, [this, IsDataLayerShown, DataLayerManager, CurrentLevel, &WorldToLevelDataLayerMap, &LoadedActors, &OutItems](const FWorldPartitionActorDesc* ActorDesc)
 						{
-							if (ActorDesc != nullptr && !ActorDesc->IsLoaded(true) && FActorDescTreeItem::ShouldDisplayInOutliner(ActorDesc))
+							if (ActorDesc != nullptr && !LoadedActors.Contains(ActorDesc->GetGuid()) && FActorDescTreeItem::ShouldDisplayInOutliner(ActorDesc))
 							{
 								for (const FName& DataLayerInstanceName : ActorDesc->GetDataLayerInstanceNames())
 								{
diff --git a/Engine/Source/Editor/SceneOutliner/Private/ActorHierarchy.cpp b/Engine/Source/Editor/SceneOutliner/Private/ActorHierarchy.cpp
--- a/Engine/Source/Editor/SceneOutliner/Private/ActorHierarchy.cpp
+++ b/Engine/Source/Editor/SceneOutliner/Private/ActorHierarchy.cpp
@@ -26,6 +26,7 @@
 #include "WorldPartition/WorldPartitionActorDesc.h"
 #include "WorldPartition/IWorldPartitionEditorModule.h"
 #include "WorldPartition/WorldPartitionSubsystem.h"
+#include "EditorLoadedActorCache.h"
 #include "ActorFolder.h"
 #include "ActorMode.h"
 
@@ -453,24 +454,30 @@ void FActorHierarchy::CreateUnloadedItems(UWorld* World, TArray<FSceneOutlinerTr
 	{
 		if (UWorldPartitionSubsystem* WorldPartitionSubsystem = UWorld::GetSubsystem<UWorldPartitionSubsystem>(World))
 		{
+			FEditorLoadedActorCache LoadedActorCache;
+			
 			const ULevelInstanceSubsystem* LevelInstanceSubsystem = World->GetSubsystem<ULevelInstanceSubsystem>();
 
-			WorldPartitionSubsystem->ForEachWorldPartition([this, LevelInstanceSubsystem, &OutItems](UWorldPartition* WorldPartition)
+			WorldPartitionSubsystem->ForEachWorldPartition([this, LevelInstanceSubsystem, &LoadedActorCache, &OutItems](UWorldPartition* WorldPartition)
 			{
+				UWorld* OuterWorld = WorldPartition->GetTypedOuter<UWorld>();
+				ULevel* OuterLevel = OuterWorld ? OuterWorld->PersistentLevel : nullptr;
+
 				// Skip unloaded actors if they are part of a non editing level instance and the outliner hides the content of level instances
 				if (!bShowingLevelInstances)
 				{
-					UWorld* OuterWorld = WorldPartition->GetTypedOuter<UWorld>();
-					ULevel* OuterLevel = OuterWorld ? OuterWorld->PersistentLevel : nullptr;
 					ILevelInstanceInterface* LevelInstance = LevelInstanceSubsystem ? LevelInstanceSubsystem->GetOwningLevelInstance(OuterLevel) : nullptr;
 					if (LevelInstance && !LevelInstance->IsEditing())
 					{
 						return true;
 					}
 				}
-				FWorldPartitionHelpers::ForEachActorDesc(WorldPartition, [this, &OutItems](const FWorldPartitionActorDesc* ActorDesc)
+
+				const TSet<FGuid>& LoadedActors = LoadedActorCache.GetLoadedActorsForLevel(OuterLevel);
+				
+				FWorldPartitionHelpers::ForEachActorDesc(WorldPartition, [this, &LoadedActors, &OutItems](const FWorldPartitionActorDesc* ActorDesc)
 				{
-					if (ActorDesc != nullptr && !ActorDesc->IsLoaded(true) && FActorDescTreeItem::ShouldDisplayInOutliner(ActorDesc))
+					if (ActorDesc != nullptr && !LoadedActors.Contains(ActorDesc->GetGuid()) && FActorDescTreeItem::ShouldDisplayInOutliner(ActorDesc))
 					{
 						if (const FSceneOutlinerTreeItemPtr ActorDescItem = Mode->CreateItemFor<FActorDescTreeItem>(FActorDescTreeItem(ActorDesc->GetGuid(), ActorDesc->GetContainer())))
 						{
diff --git a/Engine/Source/Runtime/Engine/Private/EditorLoadedActorCache.cpp b/Engine/Source/Runtime/Engine/Private/EditorLoadedActorCache.cpp
new file mode 100644
--- /dev/null
+++ b/Engine/Source/Runtime/Engine/Private/EditorLoadedActorCache.cpp
@@ -0,0 +1,35 @@
+// Copyright Epic Games, Inc. All Rights Reserved.
+
+#include "EditorLoadedActorCache.h"
+
+#if WITH_EDITOR
+
+#include "GameFramework/Actor.h"
+#include "Engine/World.h"
+#include "Engine/Level.h"
+
+FEditorLoadedActorCache::FEditorLoadedActorCache()
+{
+	// Get all Actors except CDOs including invalid ones since this is a loaded cache
+	ForEachObjectOfClass(AActor::StaticClass(), [this](UObject* InObject)
+	{
+		AActor* Actor = Cast<AActor>(InObject);
+		if (ULevel* Level = Actor->GetLevel())
+		{
+			LoadedActorsPerLevel.FindOrAdd(Level).Add(Actor->GetActorGuid());
+		}
+	}, true, RF_ClassDefaultObject, EInternalObjectFlags::None);
+}
+
+const TSet<FGuid>& FEditorLoadedActorCache::GetLoadedActorsForLevel(ULevel* InLevel) const
+{
+	if (const TSet<FGuid>* LoadedActors = LoadedActorsPerLevel.Find(InLevel))
+	{
+		return *LoadedActors;
+	}
+
+	static TSet<FGuid> EmptyGuids;
+	return EmptyGuids;
+}
+
+#endif
\ No newline at end of file
diff --git a/Engine/Source/Runtime/Engine/Public/EditorLoadedActorCache.h b/Engine/Source/Runtime/Engine/Public/EditorLoadedActorCache.h
new file mode 100644
--- /dev/null
+++ b/Engine/Source/Runtime/Engine/Public/EditorLoadedActorCache.h
@@ -0,0 +1,30 @@
+// Copyright Epic Games, Inc. All Rights Reserved.
+
+#pragma once
+
+#if WITH_EDITOR
+
+#include "CoreFwd.h"
+#include "Containers/Map.h"
+#include "Containers/Set.h"
+#include "Misc/Guid.h"
+
+class ULevel;
+class UWorld;
+
+/**
+ * This class should be used as a snapshot cache of loaded actors. 
+ * It doesn't update its internal state after creation.
+ */
+struct ENGINE_API FEditorLoadedActorCache
+{
+public:
+	FEditorLoadedActorCache();
+		
+	const TSet<FGuid>& GetLoadedActorsForLevel(ULevel* InLevel) const;
+
+private:
+	TMap<ULevel*, TSet<FGuid>> LoadedActorsPerLevel;
+};
+
+#endif
\ No newline at end of file
