Data Layers: change default from editor to runtime when creating public data layer assets.

diff --git a/Engine/Source/Editor/DataLayerEditor/Private/DataLayer/DataLayerFactory.cpp b/Engine/Source/Editor/DataLayerEditor/Private/DataLayer/DataLayerFactory.cpp
--- a/Engine/Source/Editor/DataLayerEditor/Private/DataLayer/DataLayerFactory.cpp
+++ b/Engine/Source/Editor/DataLayerEditor/Private/DataLayer/DataLayerFactory.cpp
@@ -24,5 +24,9 @@ UObject* UDataLayerFactory::FactoryCreateNew(UClass* Class, UObject* InParent, F
 {
 	UDataLayerAsset* DataLayerAsset = NewObject<UDataLayerAsset>(InParent, InName, Flags);
 	DataLayerAsset->SetDebugColor(FColor::MakeRandomSeededColor(GetTypeHash(DataLayerAsset->GetFullName())));
+	if (!DataLayerAsset->IsPrivate())
+	{
+		DataLayerAsset->SetType(EDataLayerType::Runtime);
+	}
 	return DataLayerAsset;
 }
\ No newline at end of file
