PR #10872: Fix crash in data layer editor when streaming-out a partitioned level with runtime data layers

diff --git a/Engine/Source/Editor/DataLayerEditor/Private/DataLayer/DataLayerTreeItem.cpp b/Engine/Source/Editor/DataLayerEditor/Private/DataLayer/DataLayerTreeItem.cpp
--- a/Engine/Source/Editor/DataLayerEditor/Private/DataLayer/DataLayerTreeItem.cpp
+++ b/Engine/Source/Editor/DataLayerEditor/Private/DataLayer/DataLayerTreeItem.cpp
@@ -201,7 +201,7 @@ void SDataLayerTreeLabel::Construct(const FArguments& InArgs, FDataLayerTreeItem
 		.VAlign(VAlign_Center)
 		[
 			SNew(SImage)
-			.Visibility_Lambda([this] { return (DataLayerPtr.IsValid() && DataLayerPtr->IsLocked() && !DataLayerPtr->IsReadOnly() && !DataLayerPtr->GetWorld()->IsPlayInEditor()) ? EVisibility::Visible : EVisibility::Collapsed; })
+			.Visibility_Lambda([this] { return (DataLayerPtr.IsValid() && DataLayerPtr->IsLocked() && !DataLayerPtr->IsReadOnly() && DataLayerPtr->GetWorld() && !DataLayerPtr->GetWorld()->IsPlayInEditor()) ? EVisibility::Visible : EVisibility::Collapsed; })
 			.ColorAndOpacity(FSlateColor::UseForeground())
 			.Image(FAppStyle::GetBrush(TEXT("PropertyWindow.Locked")))
 			.ToolTipText(LOCTEXT("LockedRuntimeDataLayerEditing", "Locked editing. (To allow editing, in Data Layer Outliner, go to Advanced -> Allow Runtime Data Layer Editing)"))
@@ -253,7 +253,7 @@ FText SDataLayerTreeLabel::GetDisplayText() const
 	FText SuffixText = FText::GetEmpty();
 	if (!bInEditingMode)
 	{
-		if (DataLayerInstance && DataLayerInstance->IsRuntime() && DataLayerInstance->GetWorld()->IsPlayInEditor())
+		if (DataLayerInstance && DataLayerInstance->IsRuntime() && DataLayerInstance->GetWorld() && DataLayerInstance->GetWorld()->IsPlayInEditor())
 		{
 			SuffixText = FText::Format(LOCTEXT("DataLayerRuntimeState", " ({0})"), FTextStringHelper::CreateFromBuffer(GetDataLayerRuntimeStateName(DataLayerInstance->GetEffectiveRuntimeState())));
 		}
@@ -334,36 +334,33 @@ FSlateColor SDataLayerTreeLabel::GetForegroundColor() const
 	}
 
 	const UDataLayerInstance* DataLayerInstance = DataLayerPtr.Get();
-	if (DataLayerInstance)
+	if (!DataLayerInstance || !DataLayerInstance->GetWorld())
+	{
+		return FLinearColor(0.2f, 0.2f, 0.25f);
+	}
+	if (DataLayerInstance->GetWorld()->IsPlayInEditor())
 	{
-		if (DataLayerInstance->GetWorld()->IsPlayInEditor())
+		if (DataLayerInstance->IsRuntime())
 		{
-			if (DataLayerInstance->IsRuntime())
+			EDataLayerRuntimeState State = DataLayerInstance->GetEffectiveRuntimeState();
+			switch (State)
 			{
-				EDataLayerRuntimeState State = DataLayerInstance->GetEffectiveRuntimeState();
-				switch (State)
-				{
-				case EDataLayerRuntimeState::Activated:
-					return FColorList::LimeGreen;
-				case EDataLayerRuntimeState::Loaded:
-					return FColorList::NeonBlue;
-				case EDataLayerRuntimeState::Unloaded:
-					return FColorList::DarkSlateGrey;
-				}
-			}
-			else
-			{
-				return FSceneOutlinerCommonLabelData::DarkColor;
+			case EDataLayerRuntimeState::Activated:
+				return FColorList::LimeGreen;
+			case EDataLayerRuntimeState::Loaded:
+				return FColorList::NeonBlue;
+			case EDataLayerRuntimeState::Unloaded:
+				return FColorList::DarkSlateGrey;
 			}
 		}
-		else if (DataLayerInstance->IsLocked())
+		else
 		{
 			return FSceneOutlinerCommonLabelData::DarkColor;
 		}
 	}
-	if (!DataLayerInstance || !DataLayerInstance->GetWorld())
+	else if (DataLayerInstance->IsLocked())
 	{
-		return FLinearColor(0.2f, 0.2f, 0.25f);
+		return FSceneOutlinerCommonLabelData::DarkColor;
 	}
 	if (IsInActorEditorContext())
 	{
diff --git a/Engine/Source/Runtime/Engine/Private/WorldPartition/DataLayer/DataLayerInstance.cpp b/Engine/Source/Runtime/Engine/Private/WorldPartition/DataLayer/DataLayerInstance.cpp
--- a/Engine/Source/Runtime/Engine/Private/WorldPartition/DataLayer/DataLayerInstance.cpp
+++ b/Engine/Source/Runtime/Engine/Private/WorldPartition/DataLayer/DataLayerInstance.cpp
@@ -143,7 +143,8 @@ bool UDataLayerInstance::IsLocked() const
 
 bool UDataLayerInstance::IsReadOnly() const
 {
-	return GetWorld()->IsGameWorld();
+	const UWorld* World = GetWorld();
+	return !World || World->IsGameWorld();
 }
 
 const TCHAR* UDataLayerInstance::GetDataLayerIconName() const
