Fixed rare crash in Data Layer outliner when binded to a null world

diff --git a/Engine/Source/Editor/DataLayerEditor/Private/DataLayer/DataLayerHierarchy.cpp b/Engine/Source/Editor/DataLayerEditor/Private/DataLayer/DataLayerHierarchy.cpp
--- a/Engine/Source/Editor/DataLayerEditor/Private/DataLayer/DataLayerHierarchy.cpp
+++ b/Engine/Source/Editor/DataLayerEditor/Private/DataLayer/DataLayerHierarchy.cpp
@@ -195,11 +195,17 @@ void FDataLayerHierarchy::CreateItems(TArray<FSceneOutlinerTreeItemPtr>& OutItem
 		}
 	};
 
-	if (GetOwningWorld()->IsPlayInEditor())
+	UWorld* OwningWorld = GetOwningWorld();
+	if (!OwningWorld)
+	{
+		return;
+	}
+
+	if (OwningWorld->IsPlayInEditor())
 	{
 		// PIE loops on each DataLayerManager of each registered WorldPartition. 
 		// For performance reasons, child actors are not shown.
-		UWorldPartitionSubsystem* WorldPartitionSubsystem = UWorld::GetSubsystem<UWorldPartitionSubsystem>(GetOwningWorld());
+		UWorldPartitionSubsystem* WorldPartitionSubsystem = UWorld::GetSubsystem<UWorldPartitionSubsystem>(OwningWorld);
 		WorldPartitionSubsystem->ForEachWorldPartition([&CreateDataLayerTreeItems](UWorldPartition* WorldPartition)
 		{
 			CreateDataLayerTreeItems(WorldPartition->GetDataLayerManager());
@@ -209,10 +215,10 @@ void FDataLayerHierarchy::CreateItems(TArray<FSceneOutlinerTreeItemPtr>& OutItem
 	else
 	{
 		// CurrentLevel represents the current level if different than the PersistentLevel
-		ULevel* CurrentLevel = (GetOwningWorld()->GetCurrentLevel() && !GetOwningWorld()->GetCurrentLevel()->IsPersistentLevel()) ? GetOwningWorld()->GetCurrentLevel() : nullptr;
+		ULevel* CurrentLevel = (OwningWorld->GetCurrentLevel() && !OwningWorld->GetCurrentLevel()->IsPersistentLevel()) ? OwningWorld->GetCurrentLevel() : nullptr;
 
 		// Create DataLayerTreeItems for World and for Current Level (if any).
-		CreateDataLayerTreeItems(GetOwningWorld()->GetDataLayerManager());
+		CreateDataLayerTreeItems(OwningWorld->GetDataLayerManager());
 
 		if (CurrentLevel)
 		{
@@ -221,7 +227,7 @@ void FDataLayerHierarchy::CreateItems(TArray<FSceneOutlinerTreeItemPtr>& OutItem
 
 		if (bShowDataLayerActors)
 		{
-			for (AActor* Actor : FActorRange(GetOwningWorld()))
+			for (AActor* Actor : FActorRange(OwningWorld))
 			{
 				// Consider all actors or actors part of current level (if there is one)
 				bool bConsiderActor = !CurrentLevel || (Actor->GetLevel() == CurrentLevel);
@@ -246,7 +252,7 @@ void FDataLayerHierarchy::CreateItems(TArray<FSceneOutlinerTreeItemPtr>& OutItem
 				TMap<UDataLayerInstance*, UDataLayerInstance*> WorldToLevelDataLayerMap;
 				if (CurrentLevel)
 				{
-					AWorldDataLayers* OwningWorldWorldDataLayers = GetOwningWorld()->GetWorldDataLayers();
+					AWorldDataLayers* OwningWorldWorldDataLayers = OwningWorld->GetWorldDataLayers();
 					AWorldDataLayers* CurrentLevelWorldDataLayers = CurrentLevel->GetWorldDataLayers();
 					if (OwningWorldWorldDataLayers && CurrentLevelWorldDataLayers)
 					{
@@ -266,12 +272,12 @@ void FDataLayerHierarchy::CreateItems(TArray<FSceneOutlinerTreeItemPtr>& OutItem
 					}
 				}
 
-				UWorldPartitionSubsystem* WorldPartitionSubsystem = UWorld::GetSubsystem<UWorldPartitionSubsystem>(GetOwningWorld());
-				const UDataLayerManager* DataLayerManager = GetOwningWorld()->GetDataLayerManager();
+				UWorldPartitionSubsystem* WorldPartitionSubsystem = UWorld::GetSubsystem<UWorldPartitionSubsystem>(OwningWorld);
+				const UDataLayerManager* DataLayerManager = OwningWorld->GetDataLayerManager();
 				if (WorldPartitionSubsystem && DataLayerManager)
 				{
 					FEditorLoadedActorCache LoadedActorCache;
-					const ULevelInstanceSubsystem* LevelInstanceSubsystem = UWorld::GetSubsystem<ULevelInstanceSubsystem>(GetOwningWorld());
+					const ULevelInstanceSubsystem* LevelInstanceSubsystem = UWorld::GetSubsystem<ULevelInstanceSubsystem>(OwningWorld);
 					WorldPartitionSubsystem->ForEachWorldPartition([this, DataLayerManager, CurrentLevel, LevelInstanceSubsystem, IsDataLayerShown, &WorldToLevelDataLayerMap, &LoadedActorCache, &OutItems](UWorldPartition* WorldPartition)
 					{
 						// Skip WorldPartition if it's not the one of the current level (the editing level instance)
