Changed create new data layer menu to Changed create new data layer with asset (when not using deprecated data layers) which allows to create/choose an existing data layer asset and assign it to the newly created Data Layer Instance in one operation.

diff --git a/Engine/Source/Editor/DataLayerEditor/Private/DataLayer/DataLayerMode.cpp b/Engine/Source/Editor/DataLayerEditor/Private/DataLayer/DataLayerMode.cpp
--- a/Engine/Source/Editor/DataLayerEditor/Private/DataLayer/DataLayerMode.cpp
+++ b/Engine/Source/Editor/DataLayerEditor/Private/DataLayer/DataLayerMode.cpp
@@ -15,6 +15,7 @@
 #include "Containers/EnumAsByte.h"
 #include "Containers/IndirectArray.h"
 #include "Containers/SparseArray.h"
+#include "ContentBrowserDelegates.h"
 #include "ContentBrowserModule.h"
 #include "DataLayer/DataLayerDragDropOp.h"
 #include "DataLayer/DataLayerEditorSubsystem.h"
@@ -56,6 +57,7 @@
 #include "LevelInstance/LevelInstanceSubsystem.h"
 #include "Misc/AssertionMacros.h"
 #include "Modules/ModuleManager.h"
+#include "PropertyCustomizationHelpers.h"
 #include "SDataLayerBrowser.h"
 #include "SSceneOutliner.h"
 #include "SceneOutlinerFilters.h"
@@ -1357,9 +1359,44 @@ void FDataLayerMode::RegisterContextMenu()
 				{
 					return CreateNewDataLayerInternal(InParentDataLayer, nullptr, bInIsPrivate);
 				};
-				
-				Section.AddMenuEntry("CreateNewDataLayer", LOCTEXT("CreateNewDataLayer", "Create New Data Layer"), FText(), FSlateIcon(),
-					FUIAction(FExecuteAction::CreateLambda([CreateNewEmptyDataLayer]() { CreateNewEmptyDataLayer(); })));
+
+				const AWorldDataLayers* WorldDataLayers = Mode->GetOwningWorld() ? Mode->GetOwningWorld()->GetWorldDataLayers() : nullptr;
+				if (WorldDataLayers && !WorldDataLayers->HasDeprecatedDataLayers())
+				{
+					Section.AddSubMenu("CreateNewDataLayerWithAsset", LOCTEXT("CreateNewDataLayerWithAssetSubMenu", "Create New DataLayer With Asset"), LOCTEXT("CreateNewDataLayerWithAssetSubMenu_ToolTip", "Create New DataLayer With Asset"),
+						FNewToolMenuDelegate::CreateLambda([CreateNewDataLayer, WorldDataLayers](UToolMenu* InSubMenu)
+						{
+							const bool bAllowClear = false;
+							const bool bAllowCopyPaste = false;
+							const TArray<const UClass*> AllowedClasses = { UDataLayerAsset::StaticClass() };
+							FToolMenuSection& Section = InSubMenu->AddSection("Data Layer Asset");
+							TSharedRef<SWidget> MenuWidget = PropertyCustomizationHelpers::MakeAssetPickerWithMenu(
+								FAssetData(),
+								bAllowClear,
+								bAllowCopyPaste,
+								AllowedClasses,
+								PropertyCustomizationHelpers::GetNewAssetFactoriesForClasses(AllowedClasses),
+								FOnShouldFilterAsset::CreateLambda([WorldDataLayers](const FAssetData& InAssetData)
+								{
+									UDataLayerAsset* DataLayerAsset = Cast<UDataLayerAsset>(InAssetData.GetAsset());
+									return !DataLayerAsset || WorldDataLayers->GetDataLayerInstance(DataLayerAsset);
+								}),
+								FOnAssetSelected::CreateLambda([CreateNewDataLayer](const FAssetData& InAssetData)
+								{
+									if (UDataLayerAsset* DataLayerAsset = Cast<UDataLayerAsset>(InAssetData.GetAsset()))
+									{
+										CreateNewDataLayer(nullptr, DataLayerAsset);
+									}
+								}),
+								FSimpleDelegate::CreateLambda([]() { FSlateApplication::Get().DismissAllMenus(); }));
+							Section.AddEntry(FToolMenuEntry::InitWidget("PickDataLayerAsset", MenuWidget, FText::GetEmpty(), false));
+						}));
+				}
+				else
+				{
+					Section.AddMenuEntry("CreateNewDataLayer", LOCTEXT("CreateNewDataLayer", "Create New Data Layer"), FText(), FSlateIcon(),
+						FUIAction(FExecuteAction::CreateLambda([CreateNewEmptyDataLayer]() { CreateNewEmptyDataLayer(); })));
+				}
 				
 				Section.AddMenuEntry("CreateNewDataLayerPrivate", LOCTEXT("CreateNewDataLayerPrivate", "Create New Private Data Layer"), LOCTEXT("CreateNewDataLayerPrivateToolTip", "Creates an Editor Only Data Layer that cannot be used by other Worlds."), FSlateIcon(),
 					FUIAction(
