Data Layers: implement current data layers coloration, will use the data layer debug color if a single data layers is current, otherwise white if more than one.

diff --git a/Engine/Source/Editor/DataLayerEditor/Private/DataLayer/DataLayerEditorSubsystem.cpp b/Engine/Source/Editor/DataLayerEditor/Private/DataLayer/DataLayerEditorSubsystem.cpp
--- a/Engine/Source/Editor/DataLayerEditor/Private/DataLayer/DataLayerEditorSubsystem.cpp
+++ b/Engine/Source/Editor/DataLayerEditor/Private/DataLayer/DataLayerEditorSubsystem.cpp
@@ -15,6 +15,7 @@
 #include "Engine/World.h"
 #include "EngineUtils.h"
 #include "GameFramework/Actor.h"
+#include "GameFramework/ActorPrimitiveColorHandler.h"
 #include "HAL/PlatformCrt.h"
 #include "Internationalization/Internationalization.h"
 #include "Internationalization/Text.h"
@@ -58,13 +59,13 @@ class SWidget;
 
 DEFINE_LOG_CATEGORY_STATIC(LogDataLayerEditorSubsystem, All, All);
 
-FDataLayerCreationParameters::FDataLayerCreationParameters()
-	:DataLayerAsset(nullptr),
-	WorldDataLayers(nullptr),
-	bIsPrivate(false)
-{
+static FName NAME_CurrentDataLayerColor(TEXT("CurrentDataLayerColor"));
 
-}
+FDataLayerCreationParameters::FDataLayerCreationParameters()
+	: DataLayerAsset(nullptr)
+	, WorldDataLayers(nullptr)
+	, bIsPrivate(false)
+{}
 
 //////////////////////////////////////////////////////////////////////////
 // FDataLayersBroadcast
@@ -132,6 +133,23 @@ void FDataLayersBroadcast::Initialize()
 		GEngine->OnLevelActorAdded().AddRaw(this, &FDataLayersBroadcast::OnLevelActorsAdded);
 		USelection::SelectionChangedEvent.AddRaw(this, &FDataLayersBroadcast::OnLevelSelectionChanged);
 		USelection::SelectObjectEvent.AddRaw(this, &FDataLayersBroadcast::OnLevelSelectionChanged);
+
+#if ENABLE_ACTOR_PRIMITIVE_COLOR_HANDLER
+		FActorPrimitiveColorHandler::Get().RegisterPrimitiveColorHandler(NAME_CurrentDataLayerColor, LOCTEXT("CurrentDataLayerColor", "Current Data Layer Color"), [](const UPrimitiveComponent* InPrimitiveComponent) -> FLinearColor
+		{
+			if (AActor* Actor = InPrimitiveComponent->GetOwner())
+			{
+				for (const UDataLayerInstance* DataLayerInstance : Actor->GetDataLayerInstances())
+				{
+					if (DataLayerInstance->IsInActorEditorContext())
+					{
+						return (DataLayerInstance->GetOuterWorldDataLayers()->GetActorEditorContextDataLayers().Num() == 1) ? DataLayerInstance->GetDebugColor() : FLinearColor::White;
+					}
+				}
+			}
+			return FLinearColor::Gray;
+		});
+#endif
 	}
 }
 
@@ -242,6 +260,10 @@ void UDataLayerEditorSubsystem::Deinitialize()
 
 	// Unregister the engine broadcast bridge
 	DataLayerEditorLoadingStateChanged.Remove(OnActorDataLayersEditorLoadingStateChangedEngineBridgeHandle);
+
+#if ENABLE_ACTOR_PRIMITIVE_COLOR_HANDLER
+	FActorPrimitiveColorHandler::Get().UnregisterPrimitiveColorHandler(NAME_CurrentDataLayerColor);
+#endif
 }
 
 
@@ -409,7 +431,7 @@ TArray<const UDataLayerInstance*> UDataLayerEditorSubsystem::GetDataLayerInstanc
 
 void UDataLayerEditorSubsystem::EditorMapChange()
 {
-	if (UWorld * World = GetWorld())
+	if (UWorld* World = GetWorld())
 	{
 		World->PersistentLevel->OnLoadedActorAddedToLevelEvent.AddUObject(this, &UDataLayerEditorSubsystem::OnLoadedActorAddedToLevel);
 		World->OnWorldPartitionInitialized().AddUObject(this, &UDataLayerEditorSubsystem::OnWorldPartitionInitialized);
@@ -1334,6 +1356,13 @@ void UDataLayerEditorSubsystem::BroadcastDataLayerChanged(const EDataLayerAction
 {
 	bRebuildSelectedDataLayersFromEditorSelection = true;
 	DataLayerChanged.Broadcast(Action, ChangedDataLayer, ChangedProperty);
+
+#if ENABLE_ACTOR_PRIMITIVE_COLOR_HANDLER
+	if (UWorld* World = GetWorld())
+	{
+		FActorPrimitiveColorHandler::Get().RefreshPrimitiveColorHandler(NAME_CurrentDataLayerColor, World);
+	}
+#endif
 }
 
 void UDataLayerEditorSubsystem::OnDataLayerEditorLoadingStateChanged(bool bIsFromUserChange)
