Data Layers: when double-clicking on a data layer instance, only append to current state if control is down, otherwise only apply selected data layer instance.

diff --git a/Engine/Source/Editor/DataLayerEditor/Private/DataLayer/DataLayerMode.cpp b/Engine/Source/Editor/DataLayerEditor/Private/DataLayer/DataLayerMode.cpp
--- a/Engine/Source/Editor/DataLayerEditor/Private/DataLayer/DataLayerMode.cpp
+++ b/Engine/Source/Editor/DataLayerEditor/Private/DataLayer/DataLayerMode.cpp
@@ -395,12 +395,42 @@ void FDataLayerMode::OnItemDoubleClick(FSceneOutlinerTreeItemPtr Item)
 				if (!DataLayerInstance->IsInActorEditorContext())
 				{
 					const FScopedTransaction Transaction(LOCTEXT("MakeCurrentDataLayers", "Make Current Data Layer(s)"));
+
+					if (!FSlateApplication::Get().GetModifierKeys().IsControlDown())
+					{
+						const UDataLayerManager* DataLayerManger = UDataLayerManager::GetDataLayerManager(GetOwningWorld());
+						const TArray<UDataLayerInstance*> ActorEditorContextDataLayers = DataLayerManger->GetActorEditorContextDataLayers();
+
+						for (UDataLayerInstance* DataLayerInstanceIt : ActorEditorContextDataLayers)
+						{
+							UDataLayerEditorSubsystem::Get()->RemoveFromActorEditorContext(DataLayerInstanceIt);
+						}
+					}
+
 					UDataLayerEditorSubsystem::Get()->AddToActorEditorContext(DataLayerInstance);
 				}
 				else
 				{
 					const FScopedTransaction Transaction(LOCTEXT("RemoveCurrentDataLayers", "Remove Current Data Layer(s)"));
-					UDataLayerEditorSubsystem::Get()->RemoveFromActorEditorContext(DataLayerInstance);
+
+					if (!FSlateApplication::Get().GetModifierKeys().IsControlDown())
+					{
+						const UDataLayerManager* DataLayerManger = UDataLayerManager::GetDataLayerManager(GetOwningWorld());
+						const TArray<UDataLayerInstance*> ActorEditorContextDataLayers = DataLayerManger->GetActorEditorContextDataLayers();
+						const bool bRemoveCurrent = ActorEditorContextDataLayers.Num() == 1;
+
+						for (UDataLayerInstance* DataLayerInstanceIt : ActorEditorContextDataLayers)
+						{
+							if (DataLayerInstanceIt != DataLayerInstance || bRemoveCurrent)
+							{
+								UDataLayerEditorSubsystem::Get()->RemoveFromActorEditorContext(DataLayerInstanceIt);
+							}
+						}
+					}
+					else
+					{
+						UDataLayerEditorSubsystem::Get()->RemoveFromActorEditorContext(DataLayerInstance);
+					}
 				}
 			}
 		}
