External Data Layer V1
- Allows for plugins to inject/remove content (actors) in a partitioned world
- Serves as a replacement for the experimental 'Content Bundle' feature
- Added new External Data Layer (EDL) Asset
- Added new Game Feature Action 'AddWorldPartitionContent' to control activation of an EDL asset
- Supports Data Layers inside a plugin (child of an EDL)
- Added builder to convert content bundles to EDL (UGameFeatureActionConvertContentBundleWorldPartitionBuilder)
- Feature is temporarily disabled by default and can be turned on using Editor Experimental Settings 'Enable World Partition External Data Layers' flag
#tests QAGame PIE/-game/cooked, Regression test PIE/cook of latest game map, Tested PIE/cook of latest game map converted to EDL

diff --git a/Engine/Plugins/Editor/DataValidation/Source/DataValidation/Private/WorldPartitionChangelistValidator.cpp b/Engine/Plugins/Editor/DataValidation/Source/DataValidation/Private/WorldPartitionChangelistValidator.cpp
--- a/Engine/Plugins/Editor/DataValidation/Source/DataValidation/Private/WorldPartitionChangelistValidator.cpp
+++ b/Engine/Plugins/Editor/DataValidation/Source/DataValidation/Private/WorldPartitionChangelistValidator.cpp
@@ -16,6 +16,8 @@
 #include "WorldPartition/WorldPartition.h"
 #include "WorldPartition/DataLayer/DataLayerInstanceWithAsset.h"
 #include "WorldPartition/DataLayer/WorldDataLayers.h"
+#include "WorldPartition/DataLayer/ExternalDataLayerAsset.h"
+#include "WorldPartition/DataLayer/ExternalDataLayerHelper.h"
 #include "WorldPartition/ContentBundle/ContentBundlePaths.h"
 
 #include UE_INLINE_GENERATED_CPP_BY_NAME(WorldPartitionChangelistValidator)
@@ -97,7 +99,7 @@ void UWorldPartitionChangelistValidator::ValidateActorsAndDataLayersFromChangeLi
 		{
 			if (UClass* ActorNativeClass = TryAssociateActorToMap(AssetData))
 			{
-				SubmittingWorldDataLayers = ActorNativeClass->IsChildOf<AWorldDataLayers>();
+				bSubmittingWorldDataLayers = ActorNativeClass->IsChildOf<AWorldDataLayers>();
 			}
 			else if (UClass* AssetClass = AssetData.GetClass())
 			{
@@ -135,7 +137,7 @@ void UWorldPartitionChangelistValidator::ValidateActorsAndDataLayersFromChangeLi
 		}
 	}
 
-	auto RegisterContainerToValidate = [](UWorld* InWorld, FName InContainerPackageName, FActorDescContainerInstanceCollection& OutRegisteredContainers, const FGuid& InContentBundleGuid)
+	auto RegisterContainerToValidate = [](UWorld* InWorld, FName InContainerPackageName, FActorDescContainerInstanceCollection& OutRegisteredContainers, const FGuid& InContentBundleGuid = FGuid(), const UExternalDataLayerAsset* InExternalDataLayerAsset = nullptr)
 	{
 		if (OutRegisteredContainers.Contains(InContainerPackageName))
 		{
@@ -156,12 +158,15 @@ void UWorldPartitionChangelistValidator::ValidateActorsAndDataLayersFromChangeLi
 		{
 			// Find in memory failed, load the ActorDescContainerInstance
 			ContainerInstance = NewObject<UActorDescContainerInstance>();
-			ContainerInstance->Initialize({ InContainerPackageName });
-			ContainerInstance->GetContainer()->SetContentBundleGuid(InContentBundleGuid);
+			UActorDescContainerInstance::FInitializeParams InitializeParams(InContainerPackageName);
+			InitializeParams.ContentBundleGuid = InContentBundleGuid;
+			InitializeParams.ExternalDataLayerAsset = InExternalDataLayerAsset;
+			ContainerInstance->Initialize(InitializeParams);
 		}
 		else
 		{
 			check(ContainerInstance->GetContentBundleGuid() == InContentBundleGuid);
+			check(ContainerInstance->GetExternalDataLayerAsset() == InExternalDataLayerAsset);
 		}
 
 		OutRegisteredContainers.AddContainer(ContainerInstance);
@@ -181,8 +186,9 @@ void UWorldPartitionChangelistValidator::ValidateActorsAndDataLayersFromChangeLi
 		FActorDescContainerInstanceCollection ContainersToValidate;
 
 		// Always register the main world container because content bundle containers can't be validated separately
-		RegisterContainerToValidate(World, MapPath.GetPackageName(), ContainersToValidate, FGuid());
+		RegisterContainerToValidate(World, MapPath.GetPackageName(), ContainersToValidate);
 
+		TSet<FSoftObjectPath> ProcessedExternalDataLayersForMap;
 		for (const FAssetData& ActorData : ActorsData)
 		{
 			FString ActorPackagePath = ActorData.PackagePath.ToString();
@@ -196,6 +202,26 @@ void UWorldPartitionChangelistValidator::ValidateActorsAndDataLayersFromChangeLi
 
 				RegisterContainerToValidate(World, FName(*ContentBundleContainerPackagePath), ContainersToValidate, ContentBundleGuid);
 			}
+			else
+			{
+				if (TUniquePtr<FWorldPartitionActorDesc> ActorDesc = FWorldPartitionActorDescUtils::GetActorDescriptorFromAssetData(ActorData))
+				{
+					bool bIsAlreadyInSet = false;
+					const FSoftObjectPath& ExternalDataLayerPath = ActorDesc->GetExternalDataLayerAsset();
+					if (ExternalDataLayerPath.IsValid())
+					{
+						ProcessedExternalDataLayersForMap.Add(ExternalDataLayerPath, &bIsAlreadyInSet);
+						if (!bIsAlreadyInSet)
+						{
+							if (const UExternalDataLayerAsset* ExternalDataLayerAsset = Cast<UExternalDataLayerAsset>(ExternalDataLayerPath.TryLoad()))
+							{
+								const FString EDLContainerPackagePath = FExternalDataLayerHelper::GetExternalDataLayerLevelRootPath(ExternalDataLayerAsset, MapPath.GetPackageName().ToString());
+								RegisterContainerToValidate(World, FName(*EDLContainerPackagePath), ContainersToValidate, FGuid(), ExternalDataLayerAsset);
+							}
+						}
+					}
+				}
+			}
 		}
 
 		// Build a set of Relevant Actor Guids to scope error messages to what's contained in the CL 
@@ -216,15 +242,6 @@ void UWorldPartitionChangelistValidator::ValidateActorsAndDataLayersFromChangeLi
 			.SetErrorHandler(this)
 			.SetEnableStreaming(!ULevel::GetIsStreamingDisabledFromPackage(MapPath.GetPackageName()));
 
-		ContainersToValidate.ForEachActorDescContainerInstance([&Params](const UActorDescContainerInstance* ContainerInstance)
-		{
-			for (UActorDescContainerInstance::TConstIterator<> Iterator(ContainerInstance); Iterator; ++Iterator)
-			{
-				check(!Params.ActorGuidsToContainerInstanceMap.Contains(Iterator->GetGuid()));
-				Params.ActorGuidsToContainerInstanceMap.Add(Iterator->GetGuid(), ContainerInstance);
-			}
-		});
-
 		Params.ActorDescContainerInstanceCollection = &ContainersToValidate;
 		UWorldPartition::CheckForErrors(Params);
 	}
@@ -292,18 +309,29 @@ void UWorldPartitionChangelistValidator::OnInvalidReferenceGridPlacement(const I
 												FText::FromString(ReferenceActorDescView.GetIsSpatiallyLoaded() ? *SpatiallyLoadedActor : *NonSpatiallyLoadedActor),
 												FText::FromString(GetFullActorName(ReferenceActorDescView)));
 
-		AssetFails(CurrentAsset, CurrentError);
+			AssetFails(CurrentAsset, CurrentError);
 		}
 	}
 }
 
-void UWorldPartitionChangelistValidator::OnInvalidReferenceDataLayers(const IWorldPartitionActorDescInstanceView& ActorDescView, const IWorldPartitionActorDescInstanceView& ReferenceActorDescView)
+void UWorldPartitionChangelistValidator::OnInvalidReferenceDataLayers(const IWorldPartitionActorDescInstanceView& ActorDescView, const IWorldPartitionActorDescInstanceView& ReferenceActorDescView, EDataLayerInvalidReason Reason)
 {	
 	if (Filter(ActorDescView) || Filter(ReferenceActorDescView))
 	{
-		FText CurrentError = FText::Format(LOCTEXT("DataValidation.Changelist.WorldPartition.DataLayerError", "{0} is referencing {1} but both actors are using a different set of runtime data layers."),
-											FText::FromString(GetFullActorName(ActorDescView)),
-											FText::FromString(GetFullActorName(ReferenceActorDescView)));
+		FText CurrentError;
+		switch (Reason)
+		{
+		case EDataLayerInvalidReason::ReferencedActorDifferentRuntimeDataLayers:
+			CurrentError = FText::Format(LOCTEXT("DataValidation.Changelist.WorldPartition.DataLayerError", "{0} is referencing {1} but both actors are using a different set of runtime data layers."),
+				FText::FromString(GetFullActorName(ActorDescView)),
+				FText::FromString(GetFullActorName(ReferenceActorDescView)));
+			break;
+		case EDataLayerInvalidReason::ReferencedActorDifferentExternalDataLayer:
+			CurrentError = FText::Format(LOCTEXT("DataValidation.Changelist.WorldPartition.ExternalDataLayerError", "{0} is referencing {1} but both actors are assigned to a different external data layer."),
+				FText::FromString(GetFullActorName(ActorDescView)),
+				FText::FromString(GetFullActorName(ReferenceActorDescView)));
+			break;
+		}
 
 		AssetFails(CurrentAsset, CurrentError);
 	}
@@ -345,7 +373,7 @@ void UWorldPartitionChangelistValidator::OnInvalidReferenceLevelScriptDataLayers
 
 void UWorldPartitionChangelistValidator::OnInvalidReferenceDataLayerAsset(const UDataLayerInstanceWithAsset* DataLayerInstance)
 {
-	if (SubmittingWorldDataLayers)
+	if (bSubmittingWorldDataLayers)
 	{
 		FText CurrentError = FText::Format(LOCTEXT("DataValidation.Changelist.WorldPartition.InvalidDataLayerAsset", "Data layer {0} has no data layer asset."),
 			FText::FromName(DataLayerInstance->GetDataLayerFName()));
@@ -354,11 +382,22 @@ void UWorldPartitionChangelistValidator::OnInvalidReferenceDataLayerAsset(const
 	}
 }
 
+void UWorldPartitionChangelistValidator::OnInvalidDataLayerAssetType(const UDataLayerInstanceWithAsset* DataLayerInstance, const UDataLayerAsset* DataLayerAsset)
+{
+	if (bSubmittingWorldDataLayers)
+	{
+		FText CurrentError = FText::Format(LOCTEXT("DataValidation.Changelist.WorldPartition.InvalidDataLayerAssetAsset", "Data layer {0} is not compatible with Data Layer Asset {1} type {2}."),
+			FText::FromName(DataLayerInstance->GetDataLayerFName()), FText::FromName(DataLayerAsset->GetFName()), FText::FromName(DataLayerAsset->GetClass()->GetFName()));
+
+		AssetFails(CurrentAsset, CurrentError);
+	}
+}
+
 void UWorldPartitionChangelistValidator::OnDataLayerHierarchyTypeMismatch(const UDataLayerInstance* DataLayerInstance, const UDataLayerInstance* Parent)
 {
 	if (Filter(DataLayerInstance)
 		|| Filter(Parent)
-		|| SubmittingWorldDataLayers)
+		|| bSubmittingWorldDataLayers)
 	{
 		FText CurrentError = FText::Format(LOCTEXT("DataValidation.Changelist.WorldPartition.DataLayerHierarchyTypeMismatch", "Data layer {0} is of type {1} and its parent {2} is of type {3}."),
 			FText::FromString(DataLayerInstance->GetDataLayerFullName()),
@@ -374,7 +413,7 @@ void UWorldPartitionChangelistValidator::OnDataLayerAssetConflict(const UDataLay
 {
 	if (Filter(DataLayerInstance)
 		|| Filter(ConflictingDataLayerInstance)
-		|| SubmittingWorldDataLayers)
+		|| bSubmittingWorldDataLayers)
 	{
 		FText CurrentError = FText::Format(LOCTEXT("DataValidation.Changelist.WorldPartition.DataLayerAssetConflict", "Data layer instance {0} and data layer instance {1} are both referencing data layer asset {2}."),
 			FText::FromName(DataLayerInstance->GetDataLayerFName()),
diff --git a/Engine/Plugins/Editor/DataValidation/Source/DataValidation/Private/WorldPartitionChangelistValidator.h b/Engine/Plugins/Editor/DataValidation/Source/DataValidation/Private/WorldPartitionChangelistValidator.h
--- a/Engine/Plugins/Editor/DataValidation/Source/DataValidation/Private/WorldPartitionChangelistValidator.h
+++ b/Engine/Plugins/Editor/DataValidation/Source/DataValidation/Private/WorldPartitionChangelistValidator.h
@@ -20,7 +20,7 @@ protected:
 	TSet<FGuid> RelevantActorGuids;
 	TSet<FString> RelevantDataLayerAssets;
 	UObject* CurrentAsset = nullptr;
-	bool SubmittingWorldDataLayers = false;
+	bool bSubmittingWorldDataLayers = false;
 
 	virtual bool CanValidateAsset_Implementation(const FAssetData& InAssetData, UObject* InAsset, FDataValidationContext& InContext) const override;
 	virtual EDataValidationResult ValidateLoadedAsset_Implementation(const FAssetData& InAssetData, UObject* InAsset, FDataValidationContext& InContext) override;
@@ -37,11 +37,12 @@ protected:
 	virtual void OnInvalidRuntimeGrid(const IWorldPartitionActorDescInstanceView& ActorDescView, FName GridName) override;
 	virtual void OnInvalidReference(const IWorldPartitionActorDescInstanceView& ActorDescView, const FGuid& ReferenceGuid, IWorldPartitionActorDescInstanceView* ReferenceActorDescView) override;
 	virtual void OnInvalidReferenceGridPlacement(const IWorldPartitionActorDescInstanceView& ActorDescView, const IWorldPartitionActorDescInstanceView& ReferenceActorDescView) override;
-	virtual void OnInvalidReferenceDataLayers(const IWorldPartitionActorDescInstanceView& ActorDescView, const IWorldPartitionActorDescInstanceView& ReferenceActorDescView) override;
+	virtual void OnInvalidReferenceDataLayers(const IWorldPartitionActorDescInstanceView& ActorDescView, const IWorldPartitionActorDescInstanceView& ReferenceActorDescView, EDataLayerInvalidReason Reason) override;
 	virtual void OnInvalidReferenceLevelScriptStreamed(const IWorldPartitionActorDescInstanceView& ActorDescView) override;
 	virtual void OnInvalidReferenceLevelScriptDataLayers(const IWorldPartitionActorDescInstanceView& ActorDescView) override;
 	virtual void OnInvalidReferenceRuntimeGrid(const IWorldPartitionActorDescInstanceView& ActorDescView, const IWorldPartitionActorDescInstanceView& ReferenceActorDescView) override;
 	virtual void OnInvalidReferenceDataLayerAsset(const UDataLayerInstanceWithAsset* DataLayerInstance) override;
+	virtual void OnInvalidDataLayerAssetType(const UDataLayerInstanceWithAsset* DataLayerInstance, const UDataLayerAsset* DataLayerAsset) override;
 	virtual void OnDataLayerHierarchyTypeMismatch(const UDataLayerInstance* DataLayerInstance, const UDataLayerInstance* Parent) override;
 	virtual void OnDataLayerAssetConflict(const UDataLayerInstanceWithAsset* DataLayerInstance, const UDataLayerInstanceWithAsset* ConflictingDataLayerInstance) override;
 	virtual void OnActorNeedsResave(const IWorldPartitionActorDescInstanceView& ActorDescView) override;
diff --git a/Engine/Plugins/Editor/ObjectMixer/ObjectMixer/Source/ObjectMixer/Private/Views/List/SObjectMixerEditorList.cpp b/Engine/Plugins/Editor/ObjectMixer/ObjectMixer/Source/ObjectMixer/Private/Views/List/SObjectMixerEditorList.cpp
--- a/Engine/Plugins/Editor/ObjectMixer/ObjectMixer/Source/ObjectMixer/Private/Views/List/SObjectMixerEditorList.cpp
+++ b/Engine/Plugins/Editor/ObjectMixer/ObjectMixer/Source/ObjectMixer/Private/Views/List/SObjectMixerEditorList.cpp
@@ -911,7 +911,7 @@ void SObjectMixerEditorList::CreateActorTextInfoColumns(FSceneOutlinerInitializa
 						if (const UDataLayerManager* DataLayerManager = UDataLayerManager::GetDataLayerManager(OwningWorld))
 						{
 							TSet<const UDataLayerInstance*> DataLayerInstances;
-							DataLayerInstances.Append(DataLayerManager->GetDataLayerInstances(ActorDescInstance->GetDataLayerInstanceNames()));
+							DataLayerInstances.Append(DataLayerManager->GetDataLayerInstances(ActorDescInstance->GetDataLayerInstanceNames().ToArray()));
 							if (ULevelInstanceSubsystem* LevelInstanceSubsystem = UWorld::GetSubsystem<ULevelInstanceSubsystem>(OwningWorld))
 							{
 								UWorld* OuterWorld = ActorDescContainerInstance->GetTypedOuter<UWorld>();
diff --git a/Engine/Plugins/Editor/WorldPartitionHLODUtilities/Source/Private/WorldPartition/HLOD/Utilities/WorldPartitionHLODUtilities.cpp b/Engine/Plugins/Editor/WorldPartitionHLODUtilities/Source/Private/WorldPartition/HLOD/Utilities/WorldPartitionHLODUtilities.cpp
--- a/Engine/Plugins/Editor/WorldPartitionHLODUtilities/Source/Private/WorldPartition/HLOD/Utilities/WorldPartitionHLODUtilities.cpp
+++ b/Engine/Plugins/Editor/WorldPartitionHLODUtilities/Source/Private/WorldPartition/HLOD/Utilities/WorldPartitionHLODUtilities.cpp
@@ -19,6 +19,7 @@
 #include "WorldPartition/HLOD/Builders/HLODBuilderMeshMerge.h"
 #include "WorldPartition/HLOD/Builders/HLODBuilderMeshSimplify.h"
 #include "WorldPartition/HLOD/Builders/HLODBuilderMeshApproximate.h"
+#include "ActorEditorContext/ScopedActorEditorContextSetExternalDataLayerAsset.h"
 
 #include "AssetCompilingManager.h"
 #include "BodySetupEnums.h"
@@ -202,6 +203,7 @@ TArray<AWorldPartitionHLOD*> FWorldPartitionHLODUtilities::CreateHLODActors(FHLO
 		if (bNewActor)
 		{
 			FContentBundleActivationScope Activationscope(InCreationParams.ContentBundleGuid);
+			FScopedActorEditorContextSetExternalDataLayerAsset EDLScope(InCreationParams.GetExternalDataLayerAsset());
 
 			FActorSpawnParameters SpawnParams;
 			SpawnParams.Name = HLODActorName;
@@ -209,19 +211,24 @@ TArray<AWorldPartitionHLOD*> FWorldPartitionHLODUtilities::CreateHLODActors(FHLO
 			HLODActor = InCreationParams.WorldPartition->GetWorld()->SpawnActor<AWorldPartitionHLOD>(HLODLayer->GetHLODActorClass(), SpawnParams);
 
 			check(HLODActor->GetContentBundleGuid() == InCreationParams.ContentBundleGuid);
+			check(HLODActor->GetExternalDataLayerAsset() == InCreationParams.GetExternalDataLayerAsset());
 
 			HLODActor->SetSourceCellGuid(InCreationParams.CellGuid);
 
 			// Make sure the generated HLOD actor has the same data layers as the source actors
 			for (const UDataLayerInstance* DataLayerInstance : InCreationParams.DataLayerInstances)
 			{
-				HLODActor->AddDataLayer(DataLayerInstance);
+				if (!DataLayerInstance->IsA<UExternalDataLayerInstance>())
+				{
+					HLODActor->AddDataLayer(DataLayerInstance);
+				}
 			}
 		}
 		else
 		{
 			check(HLODActor->GetSourceCellGuid() == InCreationParams.CellGuid);
 			check(HLODActor->GetClass() == HLODLayer->GetHLODActorClass());
+			check(HLODActor->GetExternalDataLayerAsset() == InCreationParams.GetExternalDataLayerAsset());
 		}
 
 		const TCHAR* DirtyReason = nullptr;
diff --git a/Engine/Plugins/Experimental/GameFeatures/Source/GameFeatures/GameFeatures.Build.cs b/Engine/Plugins/Experimental/GameFeatures/Source/GameFeatures/GameFeatures.Build.cs
--- a/Engine/Plugins/Experimental/GameFeatures/Source/GameFeatures/GameFeatures.Build.cs
+++ b/Engine/Plugins/Experimental/GameFeatures/Source/GameFeatures/GameFeatures.Build.cs
@@ -37,6 +37,7 @@ namespace UnrealBuildTool.Rules
 				PrivateDependencyModuleNames.AddRange(
 					new string[]
 					{
+						"UnrealEd",
 						"PluginUtils",
 					}
 				);
diff --git a/Engine/Plugins/Experimental/GameFeatures/Source/GameFeatures/Private/GameFeatureAction.cpp b/Engine/Plugins/Experimental/GameFeatures/Source/GameFeatures/Private/GameFeatureAction.cpp
--- a/Engine/Plugins/Experimental/GameFeatures/Source/GameFeatures/Private/GameFeatureAction.cpp
+++ b/Engine/Plugins/Experimental/GameFeatures/Source/GameFeatures/Private/GameFeatureAction.cpp
@@ -23,3 +23,15 @@ void UGameFeatureAction::OnGameFeatureActivating(FGameFeatureActivatingContext&
 	// Call older style if not overridden
 	OnGameFeatureActivating();
 }
+
+bool UGameFeatureAction::IsGameFeaturePluginRegistered() const
+{
+	UGameFeatureData* GameFeatureData = GetGameFeatureData();
+	return !!GameFeatureData ? GameFeatureData->IsGameFeaturePluginRegistered() : false;
+}
+
+bool UGameFeatureAction::IsGameFeaturePluginActive() const
+{
+	UGameFeatureData* GameFeatureData = GetGameFeatureData();
+	return !!GameFeatureData ? GameFeatureData->IsGameFeaturePluginActive() : false;
+}
diff --git a/Engine/Plugins/Experimental/GameFeatures/Source/GameFeatures/Private/GameFeatureAction_AddWPContent.cpp b/Engine/Plugins/Experimental/GameFeatures/Source/GameFeatures/Private/GameFeatureAction_AddWPContent.cpp
--- a/Engine/Plugins/Experimental/GameFeatures/Source/GameFeatures/Private/GameFeatureAction_AddWPContent.cpp
+++ b/Engine/Plugins/Experimental/GameFeatures/Source/GameFeatures/Private/GameFeatureAction_AddWPContent.cpp
@@ -27,9 +27,9 @@ UGameFeatureAction_AddWPContent::UGameFeatureAction_AddWPContent(const FObjectIn
 void UGameFeatureAction_AddWPContent::OnGameFeatureRegistering()
 {
 	Super::OnGameFeatureRegistering();
-	ContentBundleClient = FContentBundleClient::CreateClient(ContentBundleDescriptor, GetTypedOuter<UGameFeatureData>()->GetName());
-	UE_CLOG(ContentBundleClient == nullptr, LogGameFeatures, Error, TEXT("OnGameFeatureRegistering %s: Failed to create a content bundle client for %s"), *GetPathName(), *ContentBundleDescriptor->GetDisplayName())
 
+	ContentBundleClient = FContentBundleClient::CreateClient(ContentBundleDescriptor, GetTypedOuter<UGameFeatureData>()->GetName());
+	
 #if WITH_EDITOR
 	if (IsRunningCookCommandlet() && ContentBundleClient != nullptr)
 	{
diff --git a/Engine/Plugins/Experimental/GameFeatures/Source/GameFeatures/Private/GameFeatureData.cpp b/Engine/Plugins/Experimental/GameFeatures/Source/GameFeatures/Private/GameFeatureData.cpp
--- a/Engine/Plugins/Experimental/GameFeatures/Source/GameFeatures/Private/GameFeatureData.cpp
+++ b/Engine/Plugins/Experimental/GameFeatures/Source/GameFeatures/Private/GameFeatureData.cpp
@@ -12,13 +12,18 @@
 #include "DeviceProfiles/DeviceProfile.h"
 #include "DeviceProfiles/DeviceProfileFragment.h"
 #include "DeviceProfiles/DeviceProfileManager.h"
-#include "GameFeatureAction_AddWPContent.h"
-#include "WorldPartition/ContentBundle/ContentBundleDescriptor.h"
-#include "AssetRegistry/AssetData.h"
 #include "Interfaces/IPluginManager.h"
 
 #if WITH_EDITOR
+#include "Settings/EditorExperimentalSettings.h"
+#include "WorldPartition/ContentBundle/ContentBundleDescriptor.h"
+#include "WorldPartition/ContentBundle/ContentBundlePaths.h"
+#include "WorldPartition/DataLayer/ExternalDataLayerAsset.h"
+#include "WorldPartition/DataLayer/ExternalDataLayerHelper.h"
+#include "GameFeatureAction_AddWorldPartitionContent.h"
+#include "GameFeatureAction_AddWPContent.h"
 #include "Misc/DataValidation.h"
+#include "Engine/Level.h"
 #endif
 
 #include UE_INLINE_GENERATED_CPP_BY_NAME(GameFeatureData)
@@ -557,16 +562,35 @@ void UGameFeatureData::ReloadConfigs(FConfigFile& PluginConfig) const
 }
 
 #if WITH_EDITOR
-FName UGameFeatureData::GetContentBundleGuidsAssetRegistryTag()
+
+TArray<UClass*> UGameFeatureData::GetDisallowedActions() const
+{
+	TArray<UClass*> DisallowedClasses;
+
+	if (!GetDefault<UEditorExperimentalSettings>()->bEnableWorldPartitionExternalDataLayers)
+	{
+		DisallowedClasses.Add(UGameFeatureAction_AddWorldPartitionContent::StaticClass());
+	}
+
+	return DisallowedClasses;
+}
+
+FName UGameFeatureData::GetContentBundleGuidsAssetRegistryTagPrivate()
 {
 	static const FName ContentBundlesTag("ContentBundleGuids");
 	return ContentBundlesTag;
 }
 
-void UGameFeatureData::GetContentBundleGuidsFromAsset(const FAssetData& Asset, TArray<FGuid>& OutContentBundleGuids)
+FName UGameFeatureData::GetExternalDataLayerUIDsAssetRegistryTag()
+{
+	static const FName ExternalDataLayerUIDsTag("ExternalDataLayerUIDs");
+	return ExternalDataLayerUIDsTag;
+}
+
+void UGameFeatureData::GetContentBundleGuids(const FAssetData& Asset, TArray<FGuid>& OutContentBundleGuids)
 {
 	FString ContentBundleGuidsStr;
-	if (Asset.GetTagValue(GetContentBundleGuidsAssetRegistryTag(), ContentBundleGuidsStr))
+	if (Asset.GetTagValue(GetContentBundleGuidsAssetRegistryTagPrivate(), ContentBundleGuidsStr))
 	{
 		TArray<FString> ContentBundleGuidsStrArray;
 		ContentBundleGuidsStr.ParseIntoArray(ContentBundleGuidsStrArray, TEXT(","));
@@ -581,6 +605,53 @@ void UGameFeatureData::GetContentBundleGuidsFromAsset(const FAssetData& Asset, T
 	}
 }
 
+void UGameFeatureData::GetExternalDataLayerUIDs(const FAssetData& Asset, TArray<FExternalDataLayerUID>& OutExternalDataLayerUIDs)
+{
+	FString ExternalDataLayerUIDsStr;
+	if (Asset.GetTagValue(GetExternalDataLayerUIDsAssetRegistryTag(), ExternalDataLayerUIDsStr))
+	{
+		TArray<FString> ExternalDataLayerUIDStrArray;
+		ExternalDataLayerUIDsStr.ParseIntoArray(ExternalDataLayerUIDStrArray, TEXT(","));
+		for (const FString& ExternalDataLayerUIDStr : ExternalDataLayerUIDStrArray)
+		{
+			FExternalDataLayerUID ExternalDataLayerUID;
+			if (FExternalDataLayerUID::Parse(ExternalDataLayerUIDStr, ExternalDataLayerUID))
+			{
+				OutExternalDataLayerUIDs.Add(ExternalDataLayerUID);
+			}
+		}
+	}
+}
+
+void UGameFeatureData::GetDependencyDirectoriesFromAssetData(const FAssetData& AssetData, TArray<FString>& OutDependencyDirectories)
+{
+	const FString MountPoint = FPackageName::GetPackageMountPoint(AssetData.PackagePath.ToString()).ToString();
+
+	TArray<FGuid> ContentBundleGuids;
+	GetContentBundleGuids(AssetData, ContentBundleGuids);
+	for (const FGuid& ContentBundleGuid : ContentBundleGuids)
+	{
+		FString ContentBundleExternalActorPath;
+		if (ContentBundlePaths::BuildContentBundleExternalActorPath(MountPoint, ContentBundleGuid, ContentBundleExternalActorPath))
+		{
+			const FString ExternalActorPath = ULevel::GetExternalActorsPath(ContentBundleExternalActorPath);
+			OutDependencyDirectories.Add(ExternalActorPath);
+		}
+	}
+
+	TArray<FExternalDataLayerUID> ExternalDataLayerUIDs;
+	GetExternalDataLayerUIDs(AssetData, ExternalDataLayerUIDs);
+	for (const FExternalDataLayerUID& ExternalDataLayerUID : ExternalDataLayerUIDs)
+	{
+		FString ExternalDataLayerRootPath;
+		if (FExternalDataLayerHelper::BuildExternalDataLayerRootPath(MountPoint, ExternalDataLayerUID, ExternalDataLayerRootPath))
+		{
+			const FString ExternalActorsPath = ULevel::GetExternalActorsPath(ExternalDataLayerRootPath);
+			OutDependencyDirectories.Add(ExternalActorsPath);
+		}
+	}
+}
+
 void UGameFeatureData::GetAssetRegistryTags(TArray<FAssetRegistryTag>& OutTags) const
 {
 	PRAGMA_DISABLE_DEPRECATION_WARNINGS;
@@ -593,6 +664,7 @@ void UGameFeatureData::GetAssetRegistryTags(FAssetRegistryTagsContext Context) c
 	Super::GetAssetRegistryTags(Context);
 
 	TArray<FGuid> ContentBundleGuids;
+	TArray<FExternalDataLayerUID> ExternalDataLayerUIDs;
 
 	for (UGameFeatureAction* Action : Actions)
 	{
@@ -603,12 +675,26 @@ void UGameFeatureData::GetAssetRegistryTags(FAssetRegistryTagsContext Context) c
 				ContentBundleGuids.Add(ContentBundleDescriptor->GetGuid());
 			}
 		}
+
+		if (UGameFeatureAction_AddWorldPartitionContent* WPAction = Cast<UGameFeatureAction_AddWorldPartitionContent>(Action))
+		{
+			if (const UExternalDataLayerAsset* ExternalDataLayerAsset = WPAction->GetExternalDataLayerAsset())
+			{
+				ExternalDataLayerUIDs.Add(ExternalDataLayerAsset->GetUID());
+			}
+		}
 	}
 
 	if (ContentBundleGuids.Num() > 0)
 	{
 		FString ContentBundleGuidsStr = FString::JoinBy(ContentBundleGuids, TEXT(","), [&](const FGuid& Guid) { return Guid.ToString(); });
-		Context.AddTag(FAssetRegistryTag(GetContentBundleGuidsAssetRegistryTag(), ContentBundleGuidsStr, FAssetRegistryTag::TT_Hidden));
+		Context.AddTag(FAssetRegistryTag(GetContentBundleGuidsAssetRegistryTagPrivate(), ContentBundleGuidsStr, FAssetRegistryTag::TT_Hidden));
+	}
+
+	if (ExternalDataLayerUIDs.Num() > 0)
+	{
+		FString ExternalDataLayerUIDsStr = FString::JoinBy(ExternalDataLayerUIDs, TEXT(","), [&](const FExternalDataLayerUID& ExternalDataLayerUID) { return ExternalDataLayerUID.ToString(); });
+		Context.AddTag(FAssetRegistryTag(GetExternalDataLayerUIDsAssetRegistryTag(), ExternalDataLayerUIDsStr, FAssetRegistryTag::TT_Hidden));
 	}
 }
 #endif
@@ -662,5 +748,29 @@ void UGameFeatureData::GetPluginName(const UGameFeatureData* GFD, FString& Plugi
 	}
 }
 
+bool UGameFeatureData::IsGameFeaturePluginRegistered() const
+{
+	FString PluginURL;
+	FString PluginName;
+	GetPluginName(PluginName);
+	if (UGameFeaturesSubsystem::Get().GetPluginURLByName(PluginName, PluginURL))
+	{
+		return UGameFeaturesSubsystem::Get().IsGameFeaturePluginRegistered(PluginURL);
+	}
+	return false;
+}
+
+bool UGameFeatureData::IsGameFeaturePluginActive() const
+{
+	FString PluginURL;
+	FString PluginName;
+	GetPluginName(PluginName);
+	if (UGameFeaturesSubsystem::Get().GetPluginURLByName(PluginName, PluginURL))
+	{
+		return UGameFeaturesSubsystem::Get().IsGameFeaturePluginActive(PluginURL);
+	}
+	return false;
+}
+
 #undef LOCTEXT_NAMESPACE
 
diff --git a/Engine/Plugins/Experimental/GameFeatures/Source/GameFeatures/Private/GameFeatureDataAssetDependencyGatherer.cpp b/Engine/Plugins/Experimental/GameFeatures/Source/GameFeatures/Private/GameFeatureDataAssetDependencyGatherer.cpp
--- a/Engine/Plugins/Experimental/GameFeatures/Source/GameFeatures/Private/GameFeatureDataAssetDependencyGatherer.cpp
+++ b/Engine/Plugins/Experimental/GameFeatures/Source/GameFeatures/Private/GameFeatureDataAssetDependencyGatherer.cpp
@@ -7,45 +7,33 @@
 #include "AssetRegistry/ARFilter.h"
 #include "GameFeatureData.h"
 #include "AssetRegistry/AssetRegistryState.h"
-#include "WorldPartition/ContentBundle/ContentBundlePaths.h"
-#include "Engine/Level.h"
 
 // Register FGameFeatureDataAssetDependencyGatherer for UGameFeatureData class
 REGISTER_ASSETDEPENDENCY_GATHERER(FGameFeatureDataAssetDependencyGatherer, UGameFeatureData);
 
 void FGameFeatureDataAssetDependencyGatherer::GatherDependencies(const FAssetData& AssetData, const FAssetRegistryState& AssetRegistryState, TFunctionRef<FARCompiledFilter(const FARFilter&)> CompileFilterFunc, TArray<IAssetDependencyGatherer::FGathereredDependency>& OutDependencies, TArray<FString>& OutDependencyDirectories) const
 {
-	TArray<FGuid> ContentBundleGuids;
-	UGameFeatureData::GetContentBundleGuidsFromAsset(AssetData, ContentBundleGuids);
-	if (ContentBundleGuids.Num() > 0)
+	FARFilter Filter;
+	Filter.bRecursivePaths = true;
+	Filter.bIncludeOnlyOnDiskAssets = true;
+
+	TArray<FString> DependencyDirectories;
+	UGameFeatureData::GetDependencyDirectoriesFromAssetData(AssetData, DependencyDirectories);
+	for (const FString& DependencyDirectory : DependencyDirectories)
 	{
-		FARFilter Filter;
-		Filter.bRecursivePaths = true;
-		Filter.bIncludeOnlyOnDiskAssets = true;
+		OutDependencyDirectories.Add(DependencyDirectory);
+		Filter.PackagePaths.Add(*DependencyDirectory);
+	}
 
-		FName MountPoint = FPackageName::GetPackageMountPoint(AssetData.PackagePath.ToString());
-		for (const FGuid& ContentBundleGuid : ContentBundleGuids)
-		{
-			FString ContentBundleExternalActorPath;
-			if (ContentBundlePaths::BuildContentBundleExternalActorPath(MountPoint.ToString(), ContentBundleGuid, ContentBundleExternalActorPath))
-			{
-				const FString ExternalActorsPath = ULevel::GetExternalActorsPath(ContentBundleExternalActorPath);
-
-				OutDependencyDirectories.Add(ExternalActorsPath);
-				Filter.PackagePaths.Add(*ExternalActorsPath);
-			}
-		}
+	if (Filter.PackagePaths.Num() > 0)
+	{
+		TArray<FAssetData> FilteredAssets;
+		AssetRegistryState.GetAssets(CompileFilterFunc(Filter), {}, FilteredAssets, true);
 
-		if (Filter.PackagePaths.Num() > 0)
+		for (const FAssetData& FilteredAsset : FilteredAssets)
 		{
-			TArray<FAssetData> FilteredAssets;
-			AssetRegistryState.GetAssets(CompileFilterFunc(Filter), {}, FilteredAssets, true);
-
-			for (const FAssetData& FilteredAsset : FilteredAssets)
-			{
-				OutDependencies.Emplace(IAssetDependencyGatherer::FGathereredDependency{ FilteredAsset.PackageName,
-					UE::AssetRegistry::EDependencyProperty::Game | UE::AssetRegistry::EDependencyProperty::Build });
-			}
+			OutDependencies.Emplace(IAssetDependencyGatherer::FGathereredDependency{ FilteredAsset.PackageName,
+				UE::AssetRegistry::EDependencyProperty::Game | UE::AssetRegistry::EDependencyProperty::Build });
 		}
 	}
 }
diff --git a/Engine/Plugins/Experimental/GameFeatures/Source/GameFeatures/Public/GameFeatureAction.h b/Engine/Plugins/Experimental/GameFeatures/Source/GameFeatures/Public/GameFeatureAction.h
--- a/Engine/Plugins/Experimental/GameFeatures/Source/GameFeatures/Public/GameFeatureAction.h
+++ b/Engine/Plugins/Experimental/GameFeatures/Source/GameFeatures/Public/GameFeatureAction.h
@@ -39,6 +39,12 @@ public:
 	/** Called when game feature is deactivated, it may be activated again in the near future */
 	virtual void OnGameFeatureDeactivating(FGameFeatureDeactivatingContext& Context) {}
 
+	/** Returns whether the action game feature plugin is registered or not. */
+	bool IsGameFeaturePluginRegistered() const;
+
+	/** Returns whether the action game feature plugin is active or not. */
+	bool IsGameFeaturePluginActive() const;
+
 #if WITH_EDITORONLY_DATA
 	virtual void AddAdditionalAssetBundleData(FAssetBundleData& AssetBundleData) {}
 #endif
diff --git a/Engine/Plugins/Experimental/GameFeatures/Source/GameFeatures/Public/GameFeatureAction_AddWPContent.h b/Engine/Plugins/Experimental/GameFeatures/Source/GameFeatures/Public/GameFeatureAction_AddWPContent.h
--- a/Engine/Plugins/Experimental/GameFeatures/Source/GameFeatures/Public/GameFeatureAction_AddWPContent.h
+++ b/Engine/Plugins/Experimental/GameFeatures/Source/GameFeatures/Public/GameFeatureAction_AddWPContent.h
@@ -14,7 +14,7 @@ class UContentBundleDescriptor;
 /**
  *
  */
-UCLASS(meta = (DisplayName = "Add World Partition Content"))
+UCLASS(meta = (DisplayName = "Add World Partition Content (Content Bundle)"))
 class GAMEFEATURES_API UGameFeatureAction_AddWPContent : public UGameFeatureAction
 {
 	GENERATED_UCLASS_BODY()
diff --git a/Engine/Plugins/Experimental/GameFeatures/Source/GameFeatures/Public/GameFeatureAction_AddWorldPartitionContent.h b/Engine/Plugins/Experimental/GameFeatures/Source/GameFeatures/Public/GameFeatureAction_AddWorldPartitionContent.h
new file mode 100644
--- /dev/null
+++ b/Engine/Plugins/Experimental/GameFeatures/Source/GameFeatures/Public/GameFeatureAction_AddWorldPartitionContent.h
@@ -0,0 +1,51 @@
+// Copyright Epic Games, Inc. All Rights Reserved.
+
+#pragma once
+
+#include "GameFeatureAction.h"
+#include "GameFeatureAction_AddWorldPartitionContent.generated.h"
+
+class UExternalDataLayerAsset;
+
+/**
+ *
+ */
+UCLASS(meta = (DisplayName = "Add World Partition Content"))
+class GAMEFEATURES_API UGameFeatureAction_AddWorldPartitionContent : public UGameFeatureAction
+{
+	GENERATED_UCLASS_BODY()
+
+public:
+	//~ Begin UGameFeatureAction interface
+	virtual void OnGameFeatureRegistering() override;
+	virtual void OnGameFeatureUnregistering() override;
+	virtual void OnGameFeatureActivating() override;
+	virtual void OnGameFeatureDeactivating(FGameFeatureDeactivatingContext& Context) override;
+	//~ End UGameFeatureAction interface
+
+#if WITH_EDITOR
+	//~ Begin UObject interface
+	virtual void PreEditChange(FProperty* PropertyThatWillChange) override;
+	virtual void PostEditChangeProperty(FPropertyChangedEvent& PropertyChangedEvent) override;
+	virtual void PreEditUndo() override;
+	virtual void PostEditUndo() override;
+	//~ End UObject interface
+
+	const TObjectPtr<const UExternalDataLayerAsset>& GetExternalDataLayerAsset() const { return ExternalDataLayerAsset; }
+#endif
+
+private:
+#if WITH_EDITOR
+	void OnExternalDataLayerAssetChanged(const UExternalDataLayerAsset* OldAsset, const UExternalDataLayerAsset* NewAsset);
+#endif
+
+	/** Used to detect changes on the Data Layer Asset in the action. */
+	TWeakObjectPtr<const UExternalDataLayerAsset> PreEditChangeExternalDataLayerAsset;
+	TWeakObjectPtr<const UExternalDataLayerAsset> PreEditUndoExternalDataLayerAsset;
+
+	/** External Data Layer used by this action. */
+	UPROPERTY(EditAnywhere, Category = DataLayer)
+	TObjectPtr<const UExternalDataLayerAsset> ExternalDataLayerAsset;
+
+	friend class UGameFeatureActionConvertContentBundleWorldPartitionBuilder;
+};
\ No newline at end of file
diff --git a/Engine/Plugins/Experimental/GameFeatures/Source/GameFeatures/Public/GameFeatureData.h b/Engine/Plugins/Experimental/GameFeatures/Source/GameFeatures/Public/GameFeatureData.h
--- a/Engine/Plugins/Experimental/GameFeatures/Source/GameFeatures/Public/GameFeatureData.h
+++ b/Engine/Plugins/Experimental/GameFeatures/Source/GameFeatures/Public/GameFeatureData.h
@@ -9,6 +9,7 @@
 
 class FConfigFile;
 struct FPrimaryAssetTypeInfo;
+struct FExternalDataLayerUID;
 
 struct FAssetData;
 
@@ -24,6 +25,17 @@ public:
 
 #if WITH_EDITOR
 	virtual TArray<FPrimaryAssetTypeInfo>& GetPrimaryAssetTypesToScan() { return PrimaryAssetTypesToScan; }
+	virtual void GetAssetRegistryTags(FAssetRegistryTagsContext Context) const override;
+	static void GetDependencyDirectoriesFromAssetData(const FAssetData& AssetData, TArray<FString>& OutDependencyDirectories);
+
+	//~Begin deprecation
+	UE_DEPRECATED(5.4, "Implement the version that takes FAssetRegistryTagsContext instead.")
+	virtual void GetAssetRegistryTags(TArray<FAssetRegistryTag>& OutTags) const override;
+	UE_DEPRECATED(5.4, "GetContentBundleGuidsAssetRegistryTag is deprecated")
+	static FName GetContentBundleGuidsAssetRegistryTag() { return NAME_None; }
+	UE_DEPRECATED(5.4, "GetContentBundleGuidsFromAsset is deprecated, use GetDependencyDirectoriesFromAssetData")
+	static void GetContentBundleGuidsFromAsset(const FAssetData& Asset, TArray<FGuid>& OutContentBundleGuids) {}
+	//~End deprecation
 #endif //if WITH_EDITOR
 
 	/** Method to process the base ini file for the plugin during loading */
@@ -37,13 +49,11 @@ public:
 
 	void GetPluginName(FString& PluginName) const;
 
-#if WITH_EDITOR
-	static FName GetContentBundleGuidsAssetRegistryTag();
-	static void GetContentBundleGuidsFromAsset(const FAssetData& Asset, TArray<FGuid>& OutContentBundleGuids);
-	virtual void GetAssetRegistryTags(FAssetRegistryTagsContext Context) const override;
-	UE_DEPRECATED(5.4, "Implement the version that takes FAssetRegistryTagsContext instead.")
-	virtual void GetAssetRegistryTags(TArray<FAssetRegistryTag>& OutTags) const override;
-#endif
+	/** Returns whether the game feature plugin is registered or not. */
+	bool IsGameFeaturePluginRegistered() const;
+
+	/** Returns whether the game feature plugin is active or not. */
+	bool IsGameFeaturePluginActive() const;
 
 	/**
 	 * Returns the install bundle name if one exists for this plugin.
@@ -72,13 +82,23 @@ public:
 #endif
 
 private:
+#if WITH_EDITOR
+	static void GetContentBundleGuids(const FAssetData& Asset, TArray<FGuid>& OutContentBundleGuids);
+	static void GetExternalDataLayerUIDs(const FAssetData& Asset, TArray<FExternalDataLayerUID>& OutExternalDataLayerUIDs);
+	static FName GetExternalDataLayerUIDsAssetRegistryTag();
+	static FName GetContentBundleGuidsAssetRegistryTagPrivate();
+
+	UFUNCTION()
+	TArray<UClass*> GetDisallowedActions() const;
+#endif
+
 	/** Internal helper function to reload config data on objects as a result of a plugin INI being loaded */
 	void ReloadConfigs(FConfigFile& PluginConfig) const;
 
 protected:
 
 	/** List of actions to perform as this game feature is loaded/activated/deactivated/unloaded */
-	UPROPERTY(EditDefaultsOnly, Instanced, Category="Actions")
+	UPROPERTY(EditDefaultsOnly, Instanced, Category="Actions", meta = (GetDisallowedClasses = "GetDisallowedActions"))
 	TArray<TObjectPtr<UGameFeatureAction>> Actions;
 
 	/** List of asset types to scan at startup */
diff --git a/Engine/Plugins/Experimental/GameFeatures/Source/GameFeaturesEditor/GameFeaturesEditor.Build.cs b/Engine/Plugins/Experimental/GameFeatures/Source/GameFeaturesEditor/GameFeaturesEditor.Build.cs
--- a/Engine/Plugins/Experimental/GameFeatures/Source/GameFeaturesEditor/GameFeaturesEditor.Build.cs
+++ b/Engine/Plugins/Experimental/GameFeatures/Source/GameFeaturesEditor/GameFeaturesEditor.Build.cs
@@ -17,7 +17,9 @@ namespace UnrealBuildTool.Rules
 			PrivateDependencyModuleNames.AddRange(
 				new string[]
 				{
+					"AssetTools",
 					"AssetRegistry",
+					"DataLayerEditor",
 					"DeveloperSettings",
 					"Engine",
 					"ModularGameplay",
diff --git a/Engine/Plugins/Experimental/GameFeatures/Source/GameFeaturesEditor/Private/GameFeatureActionConvertContentBundleWorldPartitionBuilder.cpp b/Engine/Plugins/Experimental/GameFeatures/Source/GameFeaturesEditor/Private/GameFeatureActionConvertContentBundleWorldPartitionBuilder.cpp
new file mode 100644
--- /dev/null
+++ b/Engine/Plugins/Experimental/GameFeatures/Source/GameFeaturesEditor/Private/GameFeatureActionConvertContentBundleWorldPartitionBuilder.cpp
@@ -0,0 +1,449 @@
+// Copyright Epic Games, Inc. All Rights Reserved.
+
+#include "GameFeatureActionConvertContentBundleWorldPartitionBuilder.h"
+#include "GameFeatureAction_AddWorldPartitionContent.h"
+#include "GameFeatureAction_AddWPContent.h"
+#include "WorldPartition/ContentBundle/ContentBundleWorldSubsystem.h"
+#include "WorldPartition/ContentBundle/ContentBundleEditor.h"
+#include "WorldPartition/ContentBundle/ContentBundleDescriptor.h"
+#include "WorldPartition/ContentBundle/ContentBundleClient.h"
+#include "WorldPartition/DataLayer/ExternalDataLayerAsset.h"
+#include "WorldPartition/DataLayer/ExternalDataLayerInstance.h"
+#include "WorldPartition/DataLayer/ExternalDataLayerManager.h"
+#include "DataLayer/DataLayerEditorSubsystem.h"
+#include "DataLayer/ExternalDataLayerFactory.h"
+#include "UObject/StrongObjectPtr.h"
+#include "UObject/MetaData.h"
+#include "Commandlets/Commandlet.h"
+#include "GameFeaturesSubsystem.h"
+#include "GameFeatureData.h"
+#include "PackageSourceControlHelper.h"
+#include "SourceControlHelpers.h"
+#include "Misc/PathViews.h"
+#include "AssetSelection.h"
+#include "PackageTools.h"
+#include "IAssetTools.h"
+#include "Algo/Find.h"
+
+DEFINE_LOG_CATEGORY_STATIC(LogConvertContentBundleBuilder, All, All);
+
+UGameFeatureActionConvertContentBundleWorldPartitionBuilder::UGameFeatureActionConvertContentBundleWorldPartitionBuilder(const FObjectInitializer& ObjectInitializer)
+	: Super(ObjectInitializer)
+	, bReportOnly(false)
+	, bSkipDelete(false)
+{}
+
+bool UGameFeatureActionConvertContentBundleWorldPartitionBuilder::PreRun(UWorld* World, FPackageSourceControlHelper& PackageHelper)
+{
+	TArray<FString> Tokens, Switches;
+	TMap<FString, FString> CommandLineParams;
+	UCommandlet::ParseCommandLine(*GetBuilderArgs(), Tokens, Switches, CommandLineParams);
+
+	FString const* DestFolderPtr = CommandLineParams.Find(TEXT("DestinationFolder"));
+	DestinationFolder = DestFolderPtr ? *DestFolderPtr : FString();
+	bReportOnly = Switches.Contains(TEXT("ReportOnly"));
+	bSkipDelete = Switches.Contains(TEXT("SkipDelete"));
+
+	if (FString const* ContentBundlesToConvertString = CommandLineParams.Find(TEXT("ContentBundles")))
+	{
+		if ((*ContentBundlesToConvertString).ParseIntoArray(ContentBundlesToConvert, TEXT("+")) == 0)
+		{
+			UE_LOG(LogConvertContentBundleBuilder, Error, TEXT("Failed to parse ContentBundles argument '%s'."), **ContentBundlesToConvertString);
+			return false;
+		}
+	}
+
+	return true;
+}
+
+bool UGameFeatureActionConvertContentBundleWorldPartitionBuilder::RunInternal(UWorld* World, const FCellInfo& InCellInfo, FPackageSourceControlHelper& PackageHelper)
+{
+	TStrongObjectPtr<UExternalDataLayerFactory> ExternalDataLayerFactory(NewObject<UExternalDataLayerFactory>(GetTransientPackage()));
+
+	UContentBundleManager* ContentBundleManager = World->ContentBundleManager;
+	if (!ContentBundleManager)
+	{
+		UE_LOG(LogConvertContentBundleBuilder, Error, TEXT("World %s does not have content bundles."), *World->GetName());
+		return false;
+	}
+
+	TArray<TSharedPtr<FContentBundleEditor>> ContentBundles;
+	if (!ContentBundleManager->GetEditorContentBundle(ContentBundles))
+	{
+		UE_LOG(LogConvertContentBundleBuilder, Error, TEXT("World %s does not have content bundles."), *World->GetName());
+		return false;
+	}
+
+	UExternalDataLayerManager* ExternalDataLayerManager = UExternalDataLayerManager::GetExternalDataLayerManager(World);
+	if (!ExternalDataLayerManager)
+	{
+		UE_LOG(LogConvertContentBundleBuilder, Error, TEXT("World %s does not have an ExternalDataLayerManager"), *World->GetName());
+		return false;
+	}
+
+	UE_LOG(LogConvertContentBundleBuilder, Log, TEXT("Found %d Content bundles in World %s"), ContentBundles.Num(), *World->GetName());
+
+	TArray<TSharedPtr<FContentBundleEditor>> ContentBundlesToProcess;
+	if (!ContentBundlesToConvert.IsEmpty())
+	{
+		for (TSharedPtr<FContentBundleEditor>& ContentBundle : ContentBundles)
+		{
+			if (!ContentBundlesToConvert.Contains(ContentBundle->GetDisplayName()))
+			{
+				UE_LOG(LogConvertContentBundleBuilder, Log, TEXT("Skipping %s: Not in the list of content bundles to convert."), *ContentBundle->GetDisplayName());
+				continue;
+			}
+			ContentBundlesToProcess.Add(ContentBundle);
+		}
+	}
+	else
+	{
+		ContentBundlesToProcess = ContentBundles;
+	}
+
+	bool bIsSuccess = true;
+	for (TSharedPtr<FContentBundleEditor>& ContentBundle : ContentBundlesToProcess)
+	{
+		TSet<UPackage*> PackagesToSave;
+		TSet<FString> PackagesToDelete;
+
+		TSharedPtr<FContentBundleClient> ContentBundleClient = ContentBundle->GetClient().Pin();
+		if (!ContentBundleClient.IsValid())
+		{
+			UE_LOG(LogConvertContentBundleBuilder, Error, TEXT("Failed to access Client of Content Bundle %s"), *ContentBundle->GetDisplayName());
+			bIsSuccess = false;
+			continue;
+		}
+
+		if (ContentBundleClient->GetState() == EContentBundleClientState::Registered)
+		{
+			if (const FContentBundleBase* CB = ContentBundleManager->GetContentBundle(World, ContentBundle->GetDescriptor()->GetGuid()))
+			{
+				if (CB->GetStatus() != EContentBundleStatus::ContentInjected)
+				{
+					UE_LOG(LogConvertContentBundleBuilder, Log, TEXT("Requesting forced injection of Content Bundle %s for conversion purposes."), *ContentBundle->GetDisplayName());
+					ContentBundleClient->RequestContentInjection();
+				}
+			}
+		}
+
+		if (const FContentBundleBase* CB = ContentBundleManager->GetContentBundle(World, ContentBundle->GetDescriptor()->GetGuid()))
+		{
+			if (CB->GetStatus() != EContentBundleStatus::ContentInjected)
+			{
+				UE_LOG(LogConvertContentBundleBuilder, Log, TEXT("Skipping %s: Not injected in this world."), *ContentBundle->GetDisplayName());
+				continue;
+			}
+		}
+
+		UActorDescContainerInstance* ContentBundleContainerInstance = ContentBundle->GetActorDescContainerInstance().Get();
+		const UContentBundleDescriptor* ContentBundleDescriptor = ContentBundle->GetDescriptor();
+		UGameFeatureAction_AddWPContent* OldAddWPContentAction = ContentBundleDescriptor->GetTypedOuter<UGameFeatureAction_AddWPContent>();
+		UGameFeatureData* GameFeatureData = ContentBundleDescriptor->GetTypedOuter<UGameFeatureData>();
+		if (!OldAddWPContentAction || !GameFeatureData || !ContentBundleContainerInstance)
+		{
+			UE_LOG(LogConvertContentBundleBuilder, Error, TEXT("Skipping invalid Content Bundle %s"), *ContentBundle->GetDisplayName());
+			bIsSuccess = false;
+			continue;
+		}
+
+		if (ContentBundleContainerInstance->IsEmpty())
+		{
+			UE_LOG(LogConvertContentBundleBuilder, Log, TEXT("Skipping conversion : Empty Content Bundle %s"), *ContentBundle->GetDisplayName());
+			SkippedEmptyContentBundles.Add(ContentBundle);
+			continue;
+		}
+
+		// Find or create EDL Asset for this Content Bundle
+		UExternalDataLayerAsset* ExternalDataLayerAsset = GetOrCreateExternalDataLayerAsset(ContentBundleDescriptor, ExternalDataLayerFactory.Get(), PackagesToSave);
+		UE_CLOG(ExternalDataLayerAsset, LogConvertContentBundleBuilder, Log, TEXT("Converting Content Bundle %s using External Data Layer Asset %s."), *ContentBundle->GetDisplayName(), *ExternalDataLayerAsset->GetPathName());
+		if (!ExternalDataLayerAsset)
+		{
+			UE_LOG(LogConvertContentBundleBuilder, Error, TEXT("Failed to retrieve or create External Data Layer Asset for Content Bundle %s"), *ContentBundle->GetDisplayName());
+			bIsSuccess = false;
+			continue;
+		}
+
+		// Allow injection of External Data Layer (only used by this builder)
+		UExternalDataLayerEngineSubsystem::FForcedExternalDataLayerInjectionKey ForcedInjectionKey(World, ExternalDataLayerAsset);
+		UExternalDataLayerEngineSubsystem::Get().ForcedAllowInjection.Add(ForcedInjectionKey);
+		ON_SCOPE_EXIT { UExternalDataLayerEngineSubsystem::Get().ForcedAllowInjection.Remove(ForcedInjectionKey); };
+
+		// Find existing GameFeatureAction_AddWorldPartitionContent for this EDL Asset
+		UGameFeatureAction* const* ExistingAddWorldPartitionContent = Algo::FindByPredicate(GameFeatureData->GetActions(), [ExternalDataLayerAsset](UGameFeatureAction* Action) { return Action && Action->IsA<UGameFeatureAction_AddWorldPartitionContent>() && Cast<UGameFeatureAction_AddWorldPartitionContent>(Action)->GetExternalDataLayerAsset() == ExternalDataLayerAsset; });
+		UGameFeatureAction_AddWorldPartitionContent* AddWorldPartitionContent = ExistingAddWorldPartitionContent ? Cast<UGameFeatureAction_AddWorldPartitionContent>(*ExistingAddWorldPartitionContent) : nullptr;
+		if (!AddWorldPartitionContent)
+		{
+			// Create new GameFeatureAction_AddWorldPartitionContent for this EDL Asset
+			AddWorldPartitionContent = NewObject<UGameFeatureAction_AddWorldPartitionContent>(GameFeatureData);
+			GameFeatureData->GetMutableActionsInEditor().Add(AddWorldPartitionContent);
+			AddWorldPartitionContent->ExternalDataLayerAsset = ExternalDataLayerAsset;
+			PackagesToSave.Add(GameFeatureData->GetPackage());
+			UE_LOG(LogConvertContentBundleBuilder, Log, TEXT("Added new Action of type 'GameFeatureAction_AddWorldPartitionContent' to GameFeatureData %s using External Data Layer Asset %s while converting Content Bundle %s."), *GameFeatureData->GetName(), *ExternalDataLayerAsset->GetPathName(), *ContentBundle->GetDisplayName());
+		}
+		
+		// Manually call OnExternalDataLayerAssetChanged to register the newly created GameFeatureAction_AddWorldPartitionContent
+		AddWorldPartitionContent->OnExternalDataLayerAssetChanged(nullptr, ExternalDataLayerAsset);
+
+		// First try to find existing ExternalDataLayerInstance
+		UExternalDataLayerInstance* ExternalDataLayerInstance = ExternalDataLayerManager->GetExternalDataLayerInstance(ExternalDataLayerAsset);
+		if (!ExternalDataLayerInstance)
+		{
+			// Create AWorldDataLayers and the ExternalDataLayerInstance for this EDL Asset
+			FDataLayerCreationParameters CreationParams;
+			CreationParams.DataLayerAsset = ExternalDataLayerAsset;
+			CreationParams.WorldDataLayers = World->GetWorldDataLayers();
+			ExternalDataLayerInstance = Cast<UExternalDataLayerInstance>(UDataLayerEditorSubsystem::Get()->CreateDataLayerInstance(CreationParams));
+			UE_CLOG(ExternalDataLayerInstance, LogConvertContentBundleBuilder, Log, TEXT("Create External Data Layer Instance %s while converting Content Bundle %s."), *ExternalDataLayerAsset->GetPathName(), *ContentBundle->GetDisplayName());
+			if (!ExternalDataLayerInstance)
+			{
+				UE_LOG(LogConvertContentBundleBuilder, Error, TEXT("Failed to create External Data Layer Instance for External Data Layer Asset %s while converting Content Bundle %s"), *ExternalDataLayerAsset->GetName(), *ContentBundle->GetDisplayName());
+				bIsSuccess = false;
+				continue;
+			}
+			check(!ExternalDataLayerInstance->IsPackageExternal());
+		}
+
+		AWorldDataLayers* EDLWorldDataLayers = ExternalDataLayerInstance->GetDirectOuterWorldDataLayers();
+		if (!EDLWorldDataLayers)
+		{
+			UE_LOG(LogConvertContentBundleBuilder, Error, TEXT("Failed to find a valid WorldDataLayers actor for External Data Layer Asset %s while converting Content Bundle %s"), *ExternalDataLayerAsset->GetName(), *ContentBundle->GetDisplayName());
+			bIsSuccess = false;
+			continue;
+		}
+
+		// By default, newly create AWorldDataLayers will be marked as read-only which would block assignation of actors to the External Data Layer
+		// Temporarily remove package flag
+		bool bIsNewEDLWorldDataLayers = EDLWorldDataLayers->GetPackage()->HasAnyPackageFlags(PKG_NewlyCreated);
+		if (bIsNewEDLWorldDataLayers)
+		{
+			EDLWorldDataLayers->GetPackage()->ClearPackageFlags(PKG_NewlyCreated);
+			PackagesToSave.Add(EDLWorldDataLayers->GetPackage());
+		}
+
+		TArray<FGuid> ConvertedActorGuids;
+		TArray<FWorldPartitionReference> ActorReferences;
+		bool bContentBundleActorConversionSuccess = true;
+		// Convert actors from Content Bundle to EDL
+		for (UActorDescContainerInstance::TIterator<> It(ContentBundleContainerInstance); It; ++It)
+		{
+			FWorldPartitionReference ActorRef(ContentBundleContainerInstance, It->GetGuid());
+			if (!ActorRef.IsValid())
+			{
+				UE_LOG(LogConvertContentBundleBuilder, Error, TEXT("Failed to load actor %s(%s). Actor won't be converted."), *It->GetActorLabelOrName().ToString(), *It->GetActorPackage().ToString());
+				bContentBundleActorConversionSuccess = false;
+				break;
+			}
+
+			ActorReferences.Add(ActorRef);
+			AActor* Actor = ActorRef.GetActor();
+			UPackage* OldActorPackage = Actor->GetExternalPackage();
+			const FString OldActorPackageName = OldActorPackage->GetName();
+			if (!bSkipDelete)
+			{
+				PackagesToDelete.Add(OldActorPackageName);
+			}
+
+			// Clear Actor CB Guid
+			FSetActorContentBundleGuid(Actor, FGuid());
+
+			// Verify that the actor can be assigned to the External Data Layer
+			FText FailureReason;
+			if (!ExternalDataLayerInstance->CanAddActor(Actor, &FailureReason))
+			{
+				UE_LOG(LogConvertContentBundleBuilder, Error, TEXT("Can't create package for actor %s. %s"), *ExternalDataLayerAsset->GetName(), *Actor->GetActorNameOrLabel(), *FailureReason.ToString());
+				bContentBundleActorConversionSuccess = false;
+				break;
+			}
+
+			// Remove actor from it's old package
+			Actor->SetPackageExternal(false);
+
+			// Get all other dependant objects in the old actor package
+			TArray<UObject*> DependantObjects;
+			ForEachObjectWithPackage(OldActorPackage, [&DependantObjects](UObject* Object)
+			{
+				if (!Cast<UMetaData>(Object))
+				{
+					DependantObjects.Add(Object);
+				}
+				return true;
+			}, false);
+
+			// Create a new external package and assign it to the actor
+			const FString BaseDir = ULevel::GetExternalActorsPath(ExternalDataLayerManager->GetExternalDataLayerLevelRootPath(ExternalDataLayerAsset));
+			UPackage* NewActorPackage = ULevel::CreateActorPackage(BaseDir, Actor->GetLevel()->GetActorPackagingScheme(), Actor->GetName());
+			Actor->SetPackageExternal(true, true, NewActorPackage);
+
+			// Validation
+			check(NewActorPackage == Actor->GetExternalPackage());
+			check(NewActorPackage->GetName() == ExternalDataLayerManager->GetActorPackageName(ExternalDataLayerAsset, Actor->GetLevel(), Actor->GetName()));
+			check(NewActorPackage->GetName() != OldActorPackageName);
+
+			// Move dependant objects into the new actor package
+			for (UObject* DependantObject : DependantObjects)
+			{
+				DependantObject->Rename(nullptr, NewActorPackage, REN_NonTransactional | REN_DontCreateRedirectors | REN_ForceNoResetLoaders | REN_DoNotDirty);
+			}
+			
+			// Set External Data Layer Asset
+			if (!UDataLayerEditorSubsystem::Get()->AddActorToDataLayer(Actor, ExternalDataLayerInstance))
+			{
+				UE_LOG(LogConvertContentBundleBuilder, Error, TEXT("Failed to add actor %s(%s) to external data layer %s. Actor won't be converted."), *It->GetActorLabelOrName().ToString(), *It->GetActorPackage().ToString(), *ExternalDataLayerAsset->GetName());
+				bContentBundleActorConversionSuccess = false;
+				break;
+			}
+
+			check(Actor->GetExternalDataLayerAsset() == ExternalDataLayerAsset);
+			check(!Actor->GetContentBundleGuid().IsValid());
+			UPackage* NewPackage = Actor->GetPackage();
+			PackagesToSave.Add(NewPackage);
+			check(OldActorPackageName != NewPackage->GetName());
+			UE_LOG(LogConvertContentBundleBuilder, Log, TEXT("Converted Actor %s(%s) to %s."), *Actor->GetName(), *NewPackage->GetName(), *ExternalDataLayerAsset->GetName());
+
+			ConvertedActorGuids.Add(It->GetGuid());
+		}
+
+		// Remove actors from Content Bundle container (since we recycled the same Actor/guid)
+		for (const FGuid& ActorGuid : ConvertedActorGuids)
+		{
+			ContentBundleContainerInstance->RemoveActor(ActorGuid);
+		}
+		
+		if (bContentBundleActorConversionSuccess)
+		{
+			// Remove Content Bundle GameFeatureData action
+			GameFeatureData->GetMutableActionsInEditor().Remove(OldAddWPContentAction);
+			PackagesToSave.Add(GameFeatureData->GetPackage());
+		}
+
+		// Restore package flag
+		if (bIsNewEDLWorldDataLayers)
+		{
+			EDLWorldDataLayers->GetPackage()->SetPackageFlags(PKG_NewlyCreated);
+		}
+		
+		// Log
+		for (UPackage* PackageToSave : PackagesToSave)
+		{
+			UE_LOG(LogConvertContentBundleBuilder, Log, TEXT("Package to save: %s"), *PackageToSave->GetPathName());
+		}
+
+		for (const FString& PackageToDelete : PackagesToDelete)
+		{
+			UE_LOG(LogConvertContentBundleBuilder, Log, TEXT("Package to delete: %s"), *PackageToDelete);
+		}
+
+		if (bContentBundleActorConversionSuccess && !bReportOnly)
+		{
+			if (!UWorldPartitionBuilder::SavePackages(PackagesToSave.Array(), PackageHelper))
+			{
+				UE_LOG(LogConvertContentBundleBuilder, Error, TEXT("Failed to save packages. Conversion sanity is not guaranteed. Consult log for details."));
+				bContentBundleActorConversionSuccess = false;
+			}
+			else if (!UWorldPartitionBuilder::DeletePackages(PackagesToDelete.Array(), PackageHelper))
+			{
+				UE_LOG(LogConvertContentBundleBuilder, Error, TEXT("Failed to delete packages. Conversion sanity is not guaranteed. Consult log for details."));
+				bContentBundleActorConversionSuccess = false;
+			}
+		}
+
+		if (bContentBundleActorConversionSuccess)
+		{
+			ConvertedContentBundles.Add(ContentBundle);
+			FinalReport.Reserve(FinalReport.Num() + PackagesToSave.Num() + PackagesToDelete.Num() + 5);
+			FinalReport.Add(TEXT("------------------------------------------------------------------------"));
+			FinalReport.Add(FString::Printf(TEXT("Converted Content Bundle '%s' (%s) to EDL Asset '%s'"), *ContentBundle->GetDisplayName(), *ContentBundle->GetDescriptor()->GetGuid().ToString(), *ExternalDataLayerAsset->GetPathName()));
+			FinalReport.Add(TEXT("------------------------------------------------------------------------"));
+			FinalReport.Add(FString::Printf(TEXT("[+] Added %d packages: "), PackagesToSave.Num()));
+			for (UPackage* PackageToSave : PackagesToSave)
+			{
+				FinalReport.Add(FString::Printf(TEXT(" |- %s"), *PackageToSave->GetPathName()));
+			}
+			if (PackagesToDelete.Num())
+			{
+				FinalReport.Add(FString::Printf(TEXT("[+] Deleted %d packages: "), PackagesToDelete.Num()));
+				for (const FString& PackageToDelete : PackagesToDelete)
+				{
+					FinalReport.Add(FString::Printf(TEXT(" |- %s"), *PackageToDelete));
+				}
+			}
+		}
+		else
+		{
+			bIsSuccess = false;
+		}
+	}
+
+	if (SkippedEmptyContentBundles.Num())
+	{
+		FinalReport.Reserve(FinalReport.Num() + SkippedEmptyContentBundles.Num() + 4);
+		FinalReport.Add(TEXT("------------------------------------------------------------------------"));
+		FinalReport.Add(FString::Printf(TEXT("%d Skipped Empty Content Bundles: "), SkippedEmptyContentBundles.Num()));
+		FinalReport.Add(TEXT("------------------------------------------------------------------------"));
+		FinalReport.Add(TEXT("[+] Content Bundles: "));
+		for (const TSharedPtr<FContentBundleEditor>& CB : SkippedEmptyContentBundles)
+		{
+			FinalReport.Add(FString::Printf(TEXT(" |- %s (%s)"), *CB->GetDisplayName(), *CB->GetDescriptor()->GetGuid().ToString()));
+		}
+	}
+
+	TArray<TSharedPtr<FContentBundleEditor>> FailedContentBundles;
+	for (TSharedPtr<FContentBundleEditor>& ContentBundle : ContentBundlesToProcess)
+	{
+		if (!ConvertedContentBundles.Contains(ContentBundle) && !SkippedEmptyContentBundles.Contains(ContentBundle))
+		{
+			FailedContentBundles.Add(ContentBundle);
+		}
+	}
+
+	if (FailedContentBundles.Num())
+	{
+		FinalReport.Reserve(FinalReport.Num() + FailedContentBundles.Num() + 4);
+		FinalReport.Add(TEXT("------------------------------------------------------------------------"));
+		FinalReport.Add(FString::Printf(TEXT("%d Failed Content Bundle Conversions: "), FailedContentBundles.Num()));
+		FinalReport.Add(TEXT("------------------------------------------------------------------------"));
+		FinalReport.Add(TEXT("[+] Content Bundles: "));
+		for (const TSharedPtr<FContentBundleEditor>& CB : FailedContentBundles)
+		{
+			FinalReport.Add(FString::Printf(TEXT(" |- %s (%s)"), *CB->GetDisplayName(), *CB->GetDescriptor()->GetGuid().ToString()));
+		}
+	}
+
+	if (FinalReport.Num())
+	{
+		UE_LOG(LogConvertContentBundleBuilder, Log, TEXT("================================================================================================="));
+		UE_LOG(LogConvertContentBundleBuilder, Log, TEXT("Content Bundle Conversion : Found(%d) | Converted(%d) | Empty(%d) | Failed(%d)"), ContentBundlesToProcess.Num(), ConvertedContentBundles.Num(), SkippedEmptyContentBundles.Num(), FailedContentBundles.Num());
+		UE_LOG(LogConvertContentBundleBuilder, Log, TEXT("================================================================================================="));
+		for (const FString& Str : FinalReport)
+		{
+			UE_LOG(LogConvertContentBundleBuilder, Log, TEXT("%s"), *Str);
+		}
+		UE_LOG(LogConvertContentBundleBuilder, Log, TEXT("================================================================================================="));
+	}
+
+	return bIsSuccess;
+}
+
+UExternalDataLayerAsset* UGameFeatureActionConvertContentBundleWorldPartitionBuilder::GetOrCreateExternalDataLayerAsset(const UContentBundleDescriptor* InContentBundleDescriptor, UExternalDataLayerFactory* InExternalDataLayerFactory, TSet<UPackage*>& OutPackagesToSave) const
+{
+	const FString ContentBundleDescriptorPackage = InContentBundleDescriptor->GetPackage()->GetName();
+	const FString AssetPath = FString(FPathViews::GetMountPointNameFromPath(ContentBundleDescriptorPackage, nullptr, false)) / DestinationFolder;
+	const FString AssetName = UPackageTools::SanitizePackageName(InContentBundleDescriptor->GetDisplayName());
+	const FString PackageName = AssetPath / AssetName;
+	const FSoftObjectPath Path(FTopLevelAssetPath(FName(PackageName), FName(AssetName)).ToString());
+	const FString ObjectPath = Path.ToString();
+	if (UExternalDataLayerAsset* ExistingExternalDataLayerAsset = LoadObject<UExternalDataLayerAsset>(nullptr, *ObjectPath, nullptr, LOAD_Quiet | LOAD_NoWarn))
+	{
+		return ExistingExternalDataLayerAsset;
+	}
+	
+	UObject* Asset = IAssetTools::Get().CreateAsset(AssetName, AssetPath, UExternalDataLayerAsset::StaticClass(), InExternalDataLayerFactory);
+	UExternalDataLayerAsset* NewExternalDataLayerAsset = Asset ? CastChecked<UExternalDataLayerAsset>(Asset) : nullptr;
+	UE_CLOG(!NewExternalDataLayerAsset, LogConvertContentBundleBuilder, Error, TEXT("Failed to create external data layer asset for %s."), *InContentBundleDescriptor->GetDisplayName());
+	if (NewExternalDataLayerAsset)
+	{
+		OutPackagesToSave.Add(NewExternalDataLayerAsset->GetPackage());
+	}
+	return NewExternalDataLayerAsset;
+}
\ No newline at end of file
diff --git a/Engine/Source/Developer/FunctionalTesting/Classes/FunctionalTest.h b/Engine/Source/Developer/FunctionalTesting/Classes/FunctionalTest.h
--- a/Engine/Source/Developer/FunctionalTesting/Classes/FunctionalTest.h
+++ b/Engine/Source/Developer/FunctionalTesting/Classes/FunctionalTest.h
@@ -737,7 +737,8 @@ public:
 	virtual void EndPlay(const EEndPlayReason::Type EndPlayReason) override;
 #if WITH_EDITOR
 	virtual bool CanChangeIsSpatiallyLoadedFlag() const override { return false; }
-	virtual bool IsDataLayerTypeSupported(TSubclassOf<UDataLayerInstance> DataLayerType) const override { return false; }
+	virtual bool ActorTypeSupportsDataLayer() const override { return false; }
+	virtual bool ActorTypeSupportsExternalDataLayer() const override { return false; }
 #endif
 	// AActor interface end
 
diff --git a/Engine/Source/Editor/DataLayerEditor/Private/DataLayer/DataLayerEditorSubsystem.cpp b/Engine/Source/Editor/DataLayerEditor/Private/DataLayer/DataLayerEditorSubsystem.cpp
--- a/Engine/Source/Editor/DataLayerEditor/Private/DataLayer/DataLayerEditorSubsystem.cpp
+++ b/Engine/Source/Editor/DataLayerEditor/Private/DataLayer/DataLayerEditorSubsystem.cpp
@@ -1,12 +1,12 @@
 // Copyright Epic Games, Inc. All Rights Reserved.
 
 #include "DataLayer/DataLayerEditorSubsystem.h"
-
+#include "DataLayer/DataLayerAction.h"
+#include "ActorEditorContext/ScopedActorEditorContextSetExternalDataLayerAsset.h"
 #include "Containers/EnumAsByte.h"
 #include "Containers/StringConv.h"
 #include "Containers/UnrealString.h"
 #include "CoreGlobals.h"
-#include "DataLayer/DataLayerAction.h"
 #include "Editor.h"
 #include "Editor/EditorEngine.h"
 #include "Engine/Engine.h"
@@ -14,6 +14,9 @@
 #include "Engine/Level.h"
 #include "Engine/World.h"
 #include "EngineUtils.h"
+#include "LevelEditorViewport.h"
+#include "LevelEditorDragDropHandler.h"
+#include "ObjectTools.h"
 #include "GameFramework/Actor.h"
 #include "GameFramework/ActorPrimitiveColorHandler.h"
 #include "HAL/PlatformCrt.h"
@@ -27,6 +30,8 @@
 #include "Misc/IFilter.h"
 #include "Misc/Optional.h"
 #include "Misc/ScopedSlowTask.h"
+#include "Algo/Transform.h"
+#include "Algo/AnyOf.h"
 #include "ProfilingDebugging/CpuProfilerTrace.h"
 #include "Selection.h"
 #include "SlotBase.h"
@@ -45,13 +50,17 @@
 #include "Widgets/SBoxPanel.h"
 #include "Widgets/Text/STextBlock.h"
 #include "WorldPartition/DataLayer/DataLayerAsset.h"
-#include "WorldPartition/DataLayer/DataLayerInstanceWithAsset.h"
 #include "WorldPartition/DataLayer/DataLayerInstancePrivate.h"
+#include "WorldPartition/DataLayer/DataLayerInstanceWithAsset.h"
+#include "WorldPartition/DataLayer/DataLayerUtils.h"
 #include "WorldPartition/DataLayer/DataLayerManager.h"
 #include "WorldPartition/DataLayer/DeprecatedDataLayerInstance.h"
+#include "WorldPartition/DataLayer/ExternalDataLayerAsset.h"
+#include "WorldPartition/DataLayer/ExternalDataLayerManager.h"
+#include "WorldPartition/DataLayer/ExternalDataLayerEngineSubsystem.h"
 #include "WorldPartition/DataLayer/WorldDataLayers.h"
-#include "WorldPartition/DataLayer/DataLayerUtils.h"
 #include "WorldPartition/WorldPartition.h"
+#include "WorldPartition/ContentBundle/ContentBundleEditorSubsystemInterface.h"
 
 class SWidget;
 
@@ -59,6 +68,7 @@ class SWidget;
 
 DEFINE_LOG_CATEGORY_STATIC(LogDataLayerEditorSubsystem, All, All);
 
+static FName NAME_DataLayerColor(TEXT("DataLayerColor"));
 static FName NAME_CurrentDataLayerColor(TEXT("CurrentDataLayerColor"));
 
 FDataLayerCreationParameters::FDataLayerCreationParameters()
@@ -81,9 +91,12 @@ private:
 	void Initialize();
 	void OnEditorMapChange(uint32 MapChangeFlags = 0) { DataLayerEditorSubsystem->EditorMapChange(); }
 	void OnPostUndoRedo() { DataLayerEditorSubsystem->PostUndoRedo(); }
+	void OnCurrentLevelChanged(ULevel* InNewLevel, ULevel* InOldLevel, UWorld* InWorld) { DataLayerEditorSubsystem->EditorRefreshDataLayerBrowser(); }
 	void OnObjectPostEditChange(UObject* Object, FPropertyChangedEvent& PropertyChangedEvent);
 	void OnLevelActorsAdded(AActor* InActor) { DataLayerEditorSubsystem->InitializeNewActorDataLayers(InActor); }
 	void OnLevelSelectionChanged(UObject* InObject) { DataLayerEditorSubsystem->OnSelectionChanged(); }
+	void OnExternalDataLayerAssetRegistrationStateChanged(const UExternalDataLayerAsset* ExternalDataLayerAsset, EExternalDataLayerRegistrationState OldState, EExternalDataLayerRegistrationState NewState) { DataLayerEditorSubsystem->OnExternalDataLayerAssetRegistrationStateChanged(ExternalDataLayerAsset, OldState, NewState); }
+	TUniquePtr<FLevelEditorDragDropWorldSurrogateReferencingObject> OnLevelEditorDragDropWorldSurrogateReferencingObject(UWorld* ReferencingWorld, const FSoftObjectPath& Object) { return DataLayerEditorSubsystem->OnLevelEditorDragDropWorldSurrogateReferencingObject(ReferencingWorld, Object); }
 
 	UDataLayerEditorSubsystem* DataLayerEditorSubsystem;
 	bool bIsInitialized;
@@ -111,6 +124,7 @@ void FDataLayersBroadcast::Deinitialize()
 		{
 			FEditorDelegates::MapChange.RemoveAll(this);
 			FEditorDelegates::PostUndoRedo.RemoveAll(this);
+			FWorldDelegates::OnCurrentLevelChanged.RemoveAll(this);
 			FCoreUObjectDelegates::OnObjectPropertyChanged.RemoveAll(this);
 			if (GEngine)
 			{
@@ -118,6 +132,12 @@ void FDataLayersBroadcast::Deinitialize()
 			}
 			USelection::SelectionChangedEvent.RemoveAll(this);
 			USelection::SelectObjectEvent.RemoveAll(this);
+			UExternalDataLayerEngineSubsystem& ExternalDataLayerEngineSubsystem = UExternalDataLayerEngineSubsystem::Get();
+			ExternalDataLayerEngineSubsystem.OnExternalDataLayerAssetRegistrationStateChanged.RemoveAll(this);
+			if (ULevelEditorDragDropHandler* DragDrop = GEditor ? GEditor->GetLevelEditorDragDropHandler() : nullptr)
+			{
+				DragDrop->OnLevelEditorDragDropWorldSurrogateReferencingObject().Unbind();
+			}
 		}
 	}
 }
@@ -129,21 +149,37 @@ void FDataLayersBroadcast::Initialize()
 		bIsInitialized = true;
 		FEditorDelegates::MapChange.AddRaw(this, &FDataLayersBroadcast::OnEditorMapChange);
 		FEditorDelegates::PostUndoRedo.AddRaw(this, &FDataLayersBroadcast::OnPostUndoRedo);
+		FWorldDelegates::OnCurrentLevelChanged.AddRaw(this, &FDataLayersBroadcast::OnCurrentLevelChanged);
 		FCoreUObjectDelegates::OnObjectPropertyChanged.AddRaw(this, &FDataLayersBroadcast::OnObjectPostEditChange);
 		GEngine->OnLevelActorAdded().AddRaw(this, &FDataLayersBroadcast::OnLevelActorsAdded);
 		USelection::SelectionChangedEvent.AddRaw(this, &FDataLayersBroadcast::OnLevelSelectionChanged);
 		USelection::SelectObjectEvent.AddRaw(this, &FDataLayersBroadcast::OnLevelSelectionChanged);
+		UExternalDataLayerEngineSubsystem& ExternalDataLayerEngineSubsystem = UExternalDataLayerEngineSubsystem::Get();
+		ExternalDataLayerEngineSubsystem.OnExternalDataLayerAssetRegistrationStateChanged.AddRaw(this, &FDataLayersBroadcast::OnExternalDataLayerAssetRegistrationStateChanged);
+
+		if (ULevelEditorDragDropHandler* DragDrop = GEditor->GetLevelEditorDragDropHandler())
+		{
+			DragDrop->OnLevelEditorDragDropWorldSurrogateReferencingObject().BindRaw(this, &FDataLayersBroadcast::OnLevelEditorDragDropWorldSurrogateReferencingObject);
+		}
 
 #if ENABLE_ACTOR_PRIMITIVE_COLOR_HANDLER
 		FActorPrimitiveColorHandler::Get().RegisterPrimitiveColorHandler(NAME_CurrentDataLayerColor, LOCTEXT("CurrentDataLayerColor", "Current Data Layer Color"), [](const UPrimitiveComponent* InPrimitiveComponent) -> FLinearColor
 		{
 			if (AActor* Actor = InPrimitiveComponent->GetOwner())
 			{
+				for (const UDataLayerInstance* DataLayerInstance : Actor->GetDataLayerInstances())
+				{
+					if (DataLayerInstance->IsActorEditorContextCurrentColorized())
+					{
+						return DataLayerInstance->GetDebugColor();
+					}
+				}
+
 				for (const UDataLayerInstance* DataLayerInstance : Actor->GetDataLayerInstances())
 				{
 					if (DataLayerInstance->IsInActorEditorContext())
 					{
-						return (DataLayerInstance->GetOuterWorldDataLayers()->GetActorEditorContextDataLayers().Num() == 1) ? DataLayerInstance->GetDebugColor() : FLinearColor::White;
+						return FLinearColor::White;
 					}
 				}
 			}
@@ -192,6 +228,7 @@ UDataLayerEditorSubsystem::UDataLayerEditorSubsystem()
 : bRebuildSelectedDataLayersFromEditorSelection(false)
 , bAsyncBroadcastDataLayerChanged(false)
 , bAsyncUpdateAllActorsVisibility(false)
+, bAsyncInvalidateViewports(false)
 {}
 
 UDataLayerEditorSubsystem* UDataLayerEditorSubsystem::Get()
@@ -208,12 +245,7 @@ void UDataLayerEditorSubsystem::Initialize(FSubsystemCollectionBase& Collection)
 	// Set up the broadcast functions for DataLayerEditorSubsystem
 	DataLayersBroadcast = MakeShareable(new FDataLayersBroadcast(this));
 
-	if (UWorld* World = GetWorld())
-	{
-		World->PersistentLevel->OnLoadedActorAddedToLevelEvent.AddUObject(this, &UDataLayerEditorSubsystem::OnLoadedActorAddedToLevel);
-		World->OnWorldPartitionInitialized().AddUObject(this, &UDataLayerEditorSubsystem::OnWorldPartitionInitialized);
-		World->OnWorldPartitionUninitialized().AddUObject(this, &UDataLayerEditorSubsystem::OnWorldPartitionUninitialized);
-	}
+	UpdateRegisteredWorldDelegates();
 
 	UActorEditorContextSubsystem::Get()->RegisterClient(this);
 
@@ -266,6 +298,23 @@ void UDataLayerEditorSubsystem::Deinitialize()
 #endif
 }
 
+void UDataLayerEditorSubsystem::OnActorPreSpawnInitialization(AActor* InActor)
+{
+	if (InActor->bIsEditorPreviewActor || !InActor->IsPackageExternal())
+	{
+		return;
+	}
+
+	UWorld* OwningWorld = InActor->GetWorld();
+	if (OwningWorld && (OwningWorld == GetWorld()))
+	{
+		if (const UExternalDataLayerAsset* ExternalDataLayerAsset = GetActorEditorContextCurrentExternalDataLayer())
+		{
+			UExternalDataLayerManager* ExternalDataLayerManager = UExternalDataLayerManager::GetExternalDataLayerManager(OwningWorld);
+			ExternalDataLayerManager->OnActorPreSpawnInitialization(InActor, ExternalDataLayerAsset);
+		}
+	}
+}
 
 UWorld* UDataLayerEditorSubsystem::GetTickableGameObjectWorld() const
 {
@@ -279,7 +328,7 @@ ETickableTickType UDataLayerEditorSubsystem::GetTickableTickType() const
 
 bool UDataLayerEditorSubsystem::IsAllowedToTick() const
 {
-	return GetWorld() && (bAsyncBroadcastDataLayerChanged || bAsyncUpdateAllActorsVisibility);
+	return GetWorld() && (bAsyncBroadcastDataLayerChanged || bAsyncUpdateAllActorsVisibility || bAsyncInvalidateViewports);
 }
 
 void UDataLayerEditorSubsystem::Tick(float DeltaTime)
@@ -295,6 +344,12 @@ void UDataLayerEditorSubsystem::Tick(float DeltaTime)
 		UpdateAllActorsVisibility(false, false);
 		bAsyncUpdateAllActorsVisibility = false;
 	}
+
+	if (bAsyncInvalidateViewports)
+	{
+		GEditor->RedrawLevelEditingViewports();
+		bAsyncInvalidateViewports = false;
+	}
 }
 
 void UDataLayerEditorSubsystem::BeginDestroy()
@@ -328,7 +383,8 @@ void UDataLayerEditorSubsystem::OnExecuteActorEditorContextAction(UWorld* InWorl
 			}
 			break;
 		case EActorEditorContextAction::PushContext:
-			DataLayerManager->PushActorEditorContext();
+		case EActorEditorContextAction::PushDuplicateContext:
+			DataLayerManager->PushActorEditorContext(InType == EActorEditorContextAction::PushDuplicateContext);
 			BroadcastDataLayerChanged(EDataLayerAction::Reset, NULL, NAME_None);
 			break;
 		case EActorEditorContextAction::PopContext:
@@ -401,6 +457,11 @@ TSharedRef<SWidget> UDataLayerEditorSubsystem::GetActorEditorContextWidget(UWorl
 
 void UDataLayerEditorSubsystem::AddToActorEditorContext(UDataLayerInstance* InDataLayerInstance)
 {
+	// EDL has priority over Content Bundle
+	if (InDataLayerInstance->IsA<UExternalDataLayerInstance>())
+	{
+		IContentBundleEditorSubsystemInterface::Get()->DeactivateCurrentContentBundleEditing();
+	}
 	check(InDataLayerInstance->CanBeInActorEditorContext());
 	if (InDataLayerInstance->AddToActorEditorContext())
 	{
@@ -419,6 +480,33 @@ void UDataLayerEditorSubsystem::RemoveFromActorEditorContext(UDataLayerInstance*
 	}
 }
 
+const UExternalDataLayerAsset* UDataLayerEditorSubsystem::GetActorEditorContextCurrentExternalDataLayer() const
+{
+	UExternalDataLayerManager* ExternalDataLayerManager = UExternalDataLayerManager::GetExternalDataLayerManager(GetWorld());
+	return ExternalDataLayerManager ? ExternalDataLayerManager->GetActorEditorContextCurrentExternalDataLayer() : nullptr;
+}
+
+bool UDataLayerEditorSubsystem::SetActorEditorContextCurrentExternalDataLayer(const UExternalDataLayerAsset* InExternalDataLayerAsset)
+{
+	UDataLayerManager* DataLayerManager = UDataLayerManager::GetDataLayerManager(GetWorld());
+	if (const UExternalDataLayerInstance* ExternalDataLayerInstance = DataLayerManager ? Cast<UExternalDataLayerInstance>(DataLayerManager->GetDataLayerInstance(InExternalDataLayerAsset)) : nullptr)
+	{
+		if (ExternalDataLayerInstance->CanBeInActorEditorContext())
+		{
+			AddToActorEditorContext(const_cast<UExternalDataLayerInstance*>(ExternalDataLayerInstance));
+			return true;
+		}
+	}
+	else if (const UExternalDataLayerAsset* CurrentExternalDataLayer = GetActorEditorContextCurrentExternalDataLayer())
+	{
+		if (const UExternalDataLayerInstance* CurrentExternalDataLayerInstance = DataLayerManager ? Cast<UExternalDataLayerInstance>(DataLayerManager->GetDataLayerInstance(CurrentExternalDataLayer)) : nullptr)
+		{
+			RemoveFromActorEditorContext(const_cast<UExternalDataLayerInstance*>(CurrentExternalDataLayerInstance));
+		}
+	}
+	return false;
+}
+
 TArray<const UDataLayerInstance*> UDataLayerEditorSubsystem::GetDataLayerInstances(const TArray<const UDataLayerAsset*> DataLayerAssets) const
 {
 	if (UDataLayerManager* DataLayerManager = UDataLayerManager::GetDataLayerManager(GetWorld()))
@@ -429,14 +517,20 @@ TArray<const UDataLayerInstance*> UDataLayerEditorSubsystem::GetDataLayerInstanc
 	return TArray<const UDataLayerInstance*>();
 }
 
-void UDataLayerEditorSubsystem::EditorMapChange()
+void UDataLayerEditorSubsystem::UpdateRegisteredWorldDelegates()
 {
 	if (UWorld* World = GetWorld())
 	{
+		World->AddOnActorPreSpawnInitialization(FOnActorSpawned::FDelegate::CreateUObject(this, &UDataLayerEditorSubsystem::OnActorPreSpawnInitialization));
 		World->PersistentLevel->OnLoadedActorAddedToLevelEvent.AddUObject(this, &UDataLayerEditorSubsystem::OnLoadedActorAddedToLevel);
 		World->OnWorldPartitionInitialized().AddUObject(this, &UDataLayerEditorSubsystem::OnWorldPartitionInitialized);
 		World->OnWorldPartitionUninitialized().AddUObject(this, &UDataLayerEditorSubsystem::OnWorldPartitionUninitialized);
 	}
+}
+
+void UDataLayerEditorSubsystem::EditorMapChange()
+{
+	UpdateRegisteredWorldDelegates();
 	BroadcastDataLayerChanged(EDataLayerAction::Reset, NULL, NAME_None);
 	UpdateAllActorsVisibility(true, true);
 }
@@ -458,11 +552,38 @@ void UDataLayerEditorSubsystem::PostUndoRedo()
 // Operations on an individual actor.
 //
 ///////////////////////////////////////////////////////////////////////////////////////////////////////////////////////
-
 bool UDataLayerEditorSubsystem::IsActorValidForDataLayer(AActor* Actor)
+{
+	// This is for backward compatibility
+	TSet<TSubclassOf<UDataLayerInstance>> DataLayerInstanceClasses({ UDataLayerInstance::StaticClass() });
+	return IsActorValidForDataLayerForClasses(Actor, DataLayerInstanceClasses);
+}
+
+bool UDataLayerEditorSubsystem::IsActorValidForDataLayerInstances(AActor* Actor, const TArray<UDataLayerInstance*>& DataLayerInstances)
+{
+	TSet<TSubclassOf<UDataLayerInstance>> DataLayerInstanceClasses;
+	for (UDataLayerInstance* DataLayerInstance : DataLayerInstances)
+	{
+		DataLayerInstanceClasses.Add(DataLayerInstance->GetClass());
+	}
+	return IsActorValidForDataLayerForClasses(Actor, DataLayerInstanceClasses);
+}
+
+bool UDataLayerEditorSubsystem::IsActorValidForDataLayerForClasses(AActor* Actor, const TSet<TSubclassOf<UDataLayerInstance>>& DataLayerInstanceClasses)
 {
 	UWorld* World = Actor ? Actor->GetWorld() : nullptr;
-	return World && (World->WorldType == EWorldType::Editor) && World->IsPartitionedWorld() && Actor && Actor->SupportsDataLayerType(UDataLayerInstance::StaticClass()) && ((Actor->GetLevel() == Actor->GetWorld()->PersistentLevel) || Actor->GetLevel()->GetWorldDataLayers());
+	if (World && (World->WorldType == EWorldType::Editor) && World->IsPartitionedWorld() && Actor && ((Actor->GetLevel() == Actor->GetWorld()->PersistentLevel) || Actor->GetLevel()->GetWorldDataLayers()))
+	{
+		for (const TSubclassOf<UDataLayerInstance>& DataLayerInstanceClass : DataLayerInstanceClasses)
+		{
+			if (!Actor->SupportsDataLayerType(DataLayerInstanceClass))
+			{
+				return false;
+			}
+		}
+		return true;
+	}
+	return false;
 }
 
 void UDataLayerEditorSubsystem::OnWorldPartitionInitialized(UWorldPartition* InWorldPartition)
@@ -484,17 +605,13 @@ void UDataLayerEditorSubsystem::OnLoadedActorAddedToLevel(AActor& InActor)
 
 void UDataLayerEditorSubsystem::InitializeNewActorDataLayers(AActor* Actor)
 {
-	if (!IsActorValidForDataLayer(Actor))
-	{
-		return;
-	}
-
 	Actor->FixupDataLayers();
 
 	// update general actor visibility
 	bool bActorModified = false;
 	bool bActorSelectionChanged = false;
 	UpdateActorVisibility(Actor, bActorSelectionChanged, bActorModified, /*bActorNotifySelectionChange*/true, /*bActorRedrawViewports*/false);
+	bAsyncInvalidateViewports = true;
 }
 
 UWorld* UDataLayerEditorSubsystem::GetWorld() const
@@ -540,23 +657,23 @@ bool UDataLayerEditorSubsystem::AddActorsToDataLayer(const TArray<AActor*>& Acto
 	return AddActorsToDataLayers(Actors, { DataLayerInstance });
 }
 
-bool UDataLayerEditorSubsystem::AddActorsToDataLayers(const TArray<AActor*>& Actors, const TArray<UDataLayerInstance*>& DataLayers)
+bool UDataLayerEditorSubsystem::AddActorsToDataLayers(const TArray<AActor*>& Actors, const TArray<UDataLayerInstance*>& DataLayerInstances)
 {
 	bool bChangesOccurred = false;
 
-	if (DataLayers.Num() > 0)
+	if (DataLayerInstances.Num() > 0)
 	{
 		GEditor->GetSelectedActors()->BeginBatchSelectOperation();
 
 		for (AActor* Actor : Actors)
 		{
-			if (!IsActorValidForDataLayer(Actor))
+			if (!IsActorValidForDataLayerInstances(Actor, DataLayerInstances))
 			{
 				continue;
 			}
 
 			bool bActorWasModified = false;
-			for (const UDataLayerInstance* DataLayerInstance : DataLayers)
+			for (const UDataLayerInstance* DataLayerInstance : DataLayerInstances)
 			{
 				if (const UDataLayerInstanceWithAsset* DataLayerInstanceWithAsset = Cast<UDataLayerInstanceWithAsset>(DataLayerInstance))
 				{
@@ -662,20 +779,20 @@ bool UDataLayerEditorSubsystem::RemoveActorsFromDataLayer(const TArray<AActor*>&
 	return RemoveActorsFromDataLayers(Actors, { DataLayerInstance });
 }
 
-bool UDataLayerEditorSubsystem::RemoveActorsFromDataLayers(const TArray<AActor*>& Actors, const TArray<UDataLayerInstance*>& DataLayers)
+bool UDataLayerEditorSubsystem::RemoveActorsFromDataLayers(const TArray<AActor*>& Actors, const TArray<UDataLayerInstance*>& DataLayerInstances)
 {
 	GEditor->GetSelectedActors()->BeginBatchSelectOperation();
 
 	bool bChangesOccurred = false;
 	for (AActor* Actor : Actors)
 	{
-		if (!IsActorValidForDataLayer(Actor))
+		if (!IsActorValidForDataLayerInstances(Actor, DataLayerInstances))
 		{
 			continue;
 		}
 
 		bool bActorWasModified = false;
-		for (const UDataLayerInstance* DataLayerInstance : DataLayers)
+		for (const UDataLayerInstance* DataLayerInstance : DataLayerInstances)
 		{
 			if (Actor->RemoveDataLayer(DataLayerInstance))
 			{
@@ -750,14 +867,14 @@ bool UDataLayerEditorSubsystem::SelectActorsInDataLayer(UDataLayerInstance* Data
 	return SelectActorsInDataLayers({ DataLayerInstance }, bSelect, bNotify, bSelectEvenIfHidden, Filter);
 }
 
-bool UDataLayerEditorSubsystem::SelectActorsInDataLayers(const TArray<UDataLayerInstance*>& DataLayers, const bool bSelect, const bool bNotify, const bool bSelectEvenIfHidden)
+bool UDataLayerEditorSubsystem::SelectActorsInDataLayers(const TArray<UDataLayerInstance*>& DataLayerInstances, const bool bSelect, const bool bNotify, const bool bSelectEvenIfHidden)
 {
-	return SelectActorsInDataLayers(DataLayers, bSelect, bNotify, bSelectEvenIfHidden, nullptr);
+	return SelectActorsInDataLayers(DataLayerInstances, bSelect, bNotify, bSelectEvenIfHidden, nullptr);
 }
 
-bool UDataLayerEditorSubsystem::SelectActorsInDataLayers(const TArray<UDataLayerInstance*>& DataLayers, const bool bSelect, const bool bNotify, const bool bSelectEvenIfHidden, const TSharedPtr<FActorFilter>& Filter)
+bool UDataLayerEditorSubsystem::SelectActorsInDataLayers(const TArray<UDataLayerInstance*>& DataLayerInstances, const bool bSelect, const bool bNotify, const bool bSelectEvenIfHidden, const TSharedPtr<FActorFilter>& Filter)
 {
-	if (DataLayers.Num() == 0)
+	if (DataLayerInstances.Num() == 0)
 	{
 		return true;
 	}
@@ -768,7 +885,7 @@ bool UDataLayerEditorSubsystem::SelectActorsInDataLayers(const TArray<UDataLayer
 	// Iterate over all actors, looking for actors in the specified DataLayers.
 	for (AActor* Actor : FActorRange(GetWorld()))
 	{
-		if (!IsActorValidForDataLayer(Actor))
+		if (!IsActorValidForDataLayerInstances(Actor, DataLayerInstances))
 		{
 			continue;
 		}
@@ -778,7 +895,7 @@ bool UDataLayerEditorSubsystem::SelectActorsInDataLayers(const TArray<UDataLayer
 			continue;
 		}
 
-		for (const UDataLayerInstance* DataLayerInstance : DataLayers)
+		for (const UDataLayerInstance* DataLayerInstance : DataLayerInstances)
 		{
 			if (Actor->ContainsDataLayer(DataLayerInstance) || Actor->GetDataLayerInstancesForLevel().Contains(DataLayerInstance))
 			{
@@ -813,11 +930,6 @@ bool UDataLayerEditorSubsystem::UpdateActorVisibility(AActor* Actor, bool& bOutS
 	bOutActorModified = false;
 	bOutSelectionChanged = false;
 
-	if (!IsActorValidForDataLayer(Actor))
-	{
-		return false;
-	}
-
 	// If the actor doesn't belong to any DataLayers
 	TArray<const UDataLayerInstance*> DataLayerInstances = Actor->GetDataLayerInstances();
 	if (DataLayerInstances.IsEmpty())
@@ -827,29 +939,36 @@ bool UDataLayerEditorSubsystem::UpdateActorVisibility(AActor* Actor, bool& bOutS
 		return bOutActorModified;
 	}
 
-	bool bActorBelongsToVisibleDataLayer = false;
-	for (const UDataLayerInstance* DataLayerInstance : DataLayerInstances)
+	bool bActorVisible = false;
+	UDataLayerManager* DataLayerManager = UDataLayerManager::GetDataLayerManager(GetWorld());
+	if (const bool bActorShouldBeLoaded = DataLayerManager ? DataLayerManager->ResolveIsLoadedInEditor(Actor->GetDataLayerInstanceNames()) : true)
+	{
+		const UExternalDataLayerInstance* ExternalDataLayerInstance = nullptr;
+		if (const UExternalDataLayerAsset* ExternalDataLayerAsset = Actor->GetExternalDataLayerAsset())
 		{
-		if (DataLayerInstance->IsEffectiveVisible())
-			{
-				if (Actor->SetIsHiddenEdLayer(false))
-				{
-					bOutActorModified = true;
-				}
-				// Stop, because we found at least one visible DataLayer the actor belongs to
-				bActorBelongsToVisibleDataLayer = true;
-			break;
-			}
-	}
+			UExternalDataLayerManager* ExternalDataLayerManager = UExternalDataLayerManager::GetExternalDataLayerManager(GetWorld());
+			ExternalDataLayerInstance = ExternalDataLayerManager ? ExternalDataLayerManager->GetExternalDataLayerInstance(ExternalDataLayerAsset) : nullptr;
+			DataLayerInstances.Remove(ExternalDataLayerInstance);
+		}
 
-	// If the actor isn't part of a visible DataLayerInstance, hide and de-select it.
-	if (!bActorBelongsToVisibleDataLayer)
-	{
-		if (Actor->SetIsHiddenEdLayer(true))
+		// Actor is hidden if its external data layer is not visibile
+		const bool bActorHiddenByEDL = ExternalDataLayerInstance && !ExternalDataLayerInstance->IsEffectiveVisible();
+		if (!bActorHiddenByEDL)
 		{
-			bOutActorModified = true;
+			// Else, actor is visible if any of its data layer is visible
+			bActorVisible = DataLayerInstances.Num() ? Algo::AnyOf(DataLayerInstances, [](const UDataLayerInstance* DataLayerInstance) { return DataLayerInstance->IsEffectiveVisible(); }) : true;
 		}
+	}
+
+	const bool bIsHiddenEdLayer = !bActorVisible;
+	if (Actor->SetIsHiddenEdLayer(bIsHiddenEdLayer))
+	{
+		bOutActorModified = true;
+	}
 		
+	// If the actor is hidden, de-select it.
+	if (bIsHiddenEdLayer)
+	{
 		// If the actor was selected, mark it as unselected
 		if (Actor->IsSelected())
 		{
@@ -873,7 +992,12 @@ bool UDataLayerEditorSubsystem::UpdateActorVisibility(AActor* Actor, bool& bOutS
 		GEditor->RedrawLevelEditingViewports();
 	}
 
-	return bOutActorModified || bOutSelectionChanged;
+	if (bOutActorModified || bOutSelectionChanged)
+	{
+		bAsyncInvalidateViewports = true;
+		return true;
+	}
+	return false;;
 }
 
 bool UDataLayerEditorSubsystem::UpdateAllActorsVisibility(const bool bNotifySelectionChange, const bool bRedrawViewports)
@@ -1187,23 +1311,23 @@ bool UDataLayerEditorSubsystem::PassDataLayersFilter(UWorld* World, const FWorld
 	if (UDataLayerManager* DataLayerManager = UDataLayerManager::GetDataLayerManager(OwningWorld))
 	{
 		// If Actor is loaded and dirty, use a newly resolved DataLayerInstanceNames array
-		auto GetLatestDataLayerInstanceNames = [DataLayerManager, ActorHandle]() -> TArray<FName>
+		auto GetLatestDataLayerInstanceNames = [DataLayerManager, ActorHandle]() -> FDataLayerInstanceNames
 		{
 			if(AActor* Actor = ActorHandle.GetInstance()->GetActor(false); Actor && Actor->GetPackage()->IsDirty())
 			{
 				TUniquePtr<FWorldPartitionActorDesc> NewActorDesc = Actor->CreateActorDesc();
-				return FDataLayerUtils::ResolvedDataLayerInstanceNames(DataLayerManager, NewActorDesc.Get());
+				return FDataLayerUtils::ResolveDataLayerInstanceNames(DataLayerManager, NewActorDesc.Get());
 			}
 			else
 			{
 				return ActorHandle.GetInstance()->GetDataLayerInstanceNames();
 			}
 		};
-				
+
 		if (IsRunningCookCommandlet())
 		{
 			// When running cook commandlet, dont allow loading of actors with runtime loaded data layers
-			for (const FName& DataLayerInstanceName : GetLatestDataLayerInstanceNames())
+			for (const FName& DataLayerInstanceName : GetLatestDataLayerInstanceNames().ToArray())
 			{
 				const UDataLayerInstance* DataLayerInstance = DataLayerManager->GetDataLayerInstance(DataLayerInstanceName);
 				if (DataLayerInstance && DataLayerInstance->IsRuntime())
@@ -1215,7 +1339,7 @@ bool UDataLayerEditorSubsystem::PassDataLayersFilter(UWorld* World, const FWorld
 			return true;
 		}
 
-		return DataLayerManager->ResolveIsLoadedInEditor(GetLatestDataLayerInstanceNames());
+		return DataLayerManager->ResolveIsLoadedInEditor(GetLatestDataLayerInstanceNames().ToArray());
 	}
 
 	return true;
@@ -1266,34 +1390,78 @@ void UDataLayerEditorSubsystem::AddAllDataLayersTo(TArray<TWeakObjectPtr<UDataLa
 
 UDataLayerInstance* UDataLayerEditorSubsystem::CreateDataLayerInstance(const FDataLayerCreationParameters& Parameters)
 {
-	UDataLayerInstance* NewDataLayer = nullptr;
+	UDataLayerInstance* NewDataLayerInstance = nullptr;
 
 	UE_CLOG(!GetWorld(), LogDataLayerEditorSubsystem, Error, TEXT("%s - Failed because world in null."), ANSI_TO_TCHAR(__FUNCTION__));
+
 	if (AWorldDataLayers* WorldDataLayers = Parameters.WorldDataLayers != nullptr ? Parameters.WorldDataLayers.Get() : (GetWorld() ? GetWorld()->GetWorldDataLayers() : nullptr))
 	{
 		if (!WorldDataLayers->HasDeprecatedDataLayers())
 		{
 			if (Parameters.bIsPrivate)
 			{
-				NewDataLayer = CreateDataLayerInstance<UDataLayerInstancePrivate>(WorldDataLayers);
+				NewDataLayerInstance = CreateDataLayerInstance<UDataLayerInstancePrivate>(WorldDataLayers);
 			}
 			else
 			{
-				NewDataLayer = CreateDataLayerInstance<UDataLayerInstanceWithAsset>(WorldDataLayers, Parameters.DataLayerAsset);
+				if (UExternalDataLayerAsset* ExternalDataLayerAsset = Cast<UExternalDataLayerAsset>(Parameters.DataLayerAsset))
+				{
+					UWorld* OuterWorld = WorldDataLayers->GetTypedOuter<UWorld>();
+					UExternalDataLayerManager* ExternalDataLayerManager = UExternalDataLayerManager::GetExternalDataLayerManager(OuterWorld);
+					if (ExternalDataLayerManager && ExternalDataLayerManager->CanInjectExternalDataLayerAsset(ExternalDataLayerAsset))
+					{
+						const bool bAllowCreate = true;
+						if (AWorldDataLayers* EDLWorldDataLayers = ExternalDataLayerManager->GetWorldDataLayers(ExternalDataLayerAsset, bAllowCreate))
+						{
+							UExternalDataLayerInstance* ExternalDataLayerInstance = EDLWorldDataLayers->GetExternalDataLayerInstance(ExternalDataLayerAsset);
+							if (ensure(!ExternalDataLayerInstance))
+							{
+								// Create External Data Layer Instance
+								ExternalDataLayerInstance = EDLWorldDataLayers->CreateDataLayer<UExternalDataLayerInstance>(ExternalDataLayerAsset);
+								if (ExternalDataLayerInstance && ExternalDataLayerManager->InjectExternalDataLayer(ExternalDataLayerAsset))
+								{
+									NewDataLayerInstance = const_cast<UExternalDataLayerInstance*>(ExternalDataLayerManager->GetExternalDataLayerInstance(ExternalDataLayerAsset));
+									ensure(NewDataLayerInstance == ExternalDataLayerInstance);
+								}
+							}
+						}
+					}
+					UE_CLOG(!NewDataLayerInstance, LogDataLayerEditorSubsystem, Error, TEXT("[EDL %s] Failed to create External Data Layer Instance."), *ExternalDataLayerAsset->GetName());
+				}
+				else
+				{
+					NewDataLayerInstance = CreateDataLayerInstance<UDataLayerInstanceWithAsset>(WorldDataLayers, Parameters.DataLayerAsset);
+				}
 			}
 		}
 		else
 		{
-			NewDataLayer = CreateDataLayerInstance<UDeprecatedDataLayerInstance>(WorldDataLayers);
+			NewDataLayerInstance = CreateDataLayerInstance<UDeprecatedDataLayerInstance>(WorldDataLayers);
+		}
+	}
+
+	if (const UExternalDataLayerInstance* RootExternalDataLayerInstance = NewDataLayerInstance ? NewDataLayerInstance->GetDirectOuterWorldDataLayers()->GetRootExternalDataLayerInstance() : nullptr)
+	{
+		UExternalDataLayerInstance* Parent = const_cast<UExternalDataLayerInstance*>(RootExternalDataLayerInstance);
+		if ((NewDataLayerInstance != Parent) && !NewDataLayerInstance->SetParent(Parent))
+		{
+			FText Reason;
+			if (!NewDataLayerInstance->CanBeChildOf(Parent, &Reason))
+			{
+				UE_LOG(LogDataLayerEditorSubsystem, Warning, TEXT("Can't create Data Layer Instance %s under %s : %s"), *NewDataLayerInstance->GetDataLayerShortName(), *Parent->GetDataLayerShortName(), *Reason.ToString());
+			}
+			// Failed to root Data Layer Instance under root EDL
+			DeleteDataLayer(NewDataLayerInstance);
+			NewDataLayerInstance = nullptr;
 		}
 	}
 
-	if (NewDataLayer != nullptr)
+	if (NewDataLayerInstance)
 	{
-		BroadcastDataLayerChanged(EDataLayerAction::Add, NewDataLayer, NAME_None);
+		BroadcastDataLayerChanged(EDataLayerAction::Add, NewDataLayerInstance, NAME_None);
 	}
 	
-	return NewDataLayer;
+	return NewDataLayerInstance;
 }
 
 void UDataLayerEditorSubsystem::DeleteDataLayers(const TArray<UDataLayerInstance*>& DataLayersToDelete)
@@ -1308,12 +1476,12 @@ void UDataLayerEditorSubsystem::DeleteDataLayers(const TArray<UDataLayerInstance
 			continue;
 		}
 
-		if (!DataLayerToDelete->IsUserManaged())
+		if (!DataLayerToDelete->CanBeRemoved())
 		{
 			continue;
 		}
 
-		AWorldDataLayers* OuterWorldDataLayers = DataLayerToDelete->GetOuterWorldDataLayers();
+		AWorldDataLayers* OuterWorldDataLayers = DataLayerToDelete->GetDirectOuterWorldDataLayers();
 		if (OuterWorldDataLayers->RemoveDataLayer(DataLayerToDelete))
 		{
 			DeletedDataLayerInstances.Add(DataLayerToDelete);
@@ -1334,12 +1502,12 @@ void UDataLayerEditorSubsystem::DeleteDataLayer(UDataLayerInstance* DataLayerToD
 		return;
 	}
 	
-	if (!DataLayerToDelete->IsUserManaged())
+	if (!DataLayerToDelete->CanBeRemoved())
 	{
 		return;
 	}
 
-	AWorldDataLayers* OuterWorldDataLayers = DataLayerToDelete->GetOuterWorldDataLayers();
+	AWorldDataLayers* OuterWorldDataLayers = DataLayerToDelete->GetDirectOuterWorldDataLayers();
 	if (OuterWorldDataLayers->RemoveDataLayer(DataLayerToDelete))
 	{
 		BroadcastDataLayerChanged(EDataLayerAction::Delete, DataLayerToDelete, NAME_None);
@@ -1388,6 +1556,74 @@ void UDataLayerEditorSubsystem::OnSelectionChanged()
 	bRebuildSelectedDataLayersFromEditorSelection = true;
 }
 
+const UExternalDataLayerAsset* UDataLayerEditorSubsystem::GetReferencingWorldSurrogateObjectForObject(UWorld* ReferencingWorld, const FSoftObjectPath& ObjectPath)
+{
+	UExternalDataLayerManager* ExternalDataLayerManager = UExternalDataLayerManager::GetExternalDataLayerManager(ReferencingWorld);
+	if (const UExternalDataLayerAsset* ExternalDataLayerAsset = ExternalDataLayerManager ? ExternalDataLayerManager->GetMatchingExternalDataLayerAssetForObjectPath(ObjectPath) : nullptr)
+	{
+		const UDataLayerInstance* ExternalDataLayerInstance = UDataLayerManager::GetDataLayerManager(ReferencingWorld)->GetDataLayerInstance(ExternalDataLayerAsset);
+		if (ExternalDataLayerInstance && ExternalDataLayerInstance->CanBeInActorEditorContext())
+		{
+			return ExternalDataLayerAsset;
+		}
+	}
+	return nullptr;
+}
+
+struct FExternalDataLayerWorldSurrogateReferencingObject : public FLevelEditorDragDropWorldSurrogateReferencingObject
+{
+public:
+	FExternalDataLayerWorldSurrogateReferencingObject(const UObject* InSurrogateObject) 
+		: FLevelEditorDragDropWorldSurrogateReferencingObject(InSurrogateObject)
+	{}
+
+	virtual bool OnPreDropObjects(UWorld* World, const TArray<UObject*>& DroppedObjects)
+	{
+		if (!FLevelEditorDragDropWorldSurrogateReferencingObject::OnPreDropObjects(World, DroppedObjects))
+		{
+			return false;
+		}
+		
+		ULevel* CurrentLevel = World->GetCurrentLevel();
+		UWorld* CurrentLevelOuterWorld = CurrentLevel ? Cast<UWorld>(CurrentLevel->GetOuter()) : nullptr;
+		UWorld* ReferencingWorld = CurrentLevelOuterWorld ? CurrentLevelOuterWorld : World;
+		
+		const UExternalDataLayerAsset* ExternalDataLayerAsset = UDataLayerEditorSubsystem::GetReferencingWorldSurrogateObjectForObject(ReferencingWorld, FSoftObjectPath(DroppedObjects[0]));
+		EDLContext = ExternalDataLayerAsset ? TUniquePtr<FScopedActorEditorContextSetExternalDataLayerAsset>(new FScopedActorEditorContextSetExternalDataLayerAsset(ExternalDataLayerAsset)) : nullptr;
+		return EDLContext.IsValid() && (UDataLayerEditorSubsystem::Get()->GetActorEditorContextCurrentExternalDataLayer() == ExternalDataLayerAsset);
+	}
+
+	virtual bool OnPostDropObjects(UWorld* World, const TArray<UObject*>& DroppedObjects)
+	{
+		if (!FLevelEditorDragDropWorldSurrogateReferencingObject::OnPostDropObjects(World, DroppedObjects))
+		{
+			return false;
+		}
+		EDLContext.Reset();
+		return true;
+	}
+
+private:
+	TUniquePtr<FScopedActorEditorContextSetExternalDataLayerAsset> EDLContext;
+};
+
+TUniquePtr<FLevelEditorDragDropWorldSurrogateReferencingObject> UDataLayerEditorSubsystem::OnLevelEditorDragDropWorldSurrogateReferencingObject(UWorld* ReferencingWorld, const FSoftObjectPath& Object)
+{
+	if (const UExternalDataLayerAsset* ExternalDataLayerAsset = GetReferencingWorldSurrogateObjectForObject(ReferencingWorld, Object))
+	{
+		return MakeUnique<FExternalDataLayerWorldSurrogateReferencingObject>(ExternalDataLayerAsset);
+	}
+	return nullptr;
+}
+
+void UDataLayerEditorSubsystem::OnExternalDataLayerAssetRegistrationStateChanged(const UExternalDataLayerAsset* ExternalDataLayerAsset, EExternalDataLayerRegistrationState OldState, EExternalDataLayerRegistrationState NewState)
+{
+	if (ExternalDataLayerAsset)
+	{
+		EditorRefreshDataLayerBrowser();
+	}
+}
+
 const TSet<TWeakObjectPtr<const UDataLayerInstance>>& UDataLayerEditorSubsystem::GetSelectedDataLayersFromEditorSelection() const
 {
 	if (bRebuildSelectedDataLayersFromEditorSelection)
diff --git a/Engine/Source/Editor/DataLayerEditor/Private/DataLayer/DataLayerFactory.cpp b/Engine/Source/Editor/DataLayerEditor/Private/DataLayer/DataLayerFactory.cpp
--- a/Engine/Source/Editor/DataLayerEditor/Private/DataLayer/DataLayerFactory.cpp
+++ b/Engine/Source/Editor/DataLayerEditor/Private/DataLayer/DataLayerFactory.cpp
@@ -7,10 +7,6 @@
 #include "Templates/SubclassOf.h"
 #include "WorldPartition/DataLayer/DataLayerAsset.h"
 
-class FFeedbackContext;
-class UClass;
-class UObject;
-
 UDataLayerFactory::UDataLayerFactory(const FObjectInitializer& ObjectInitializer)
 	: Super(ObjectInitializer)
 {
@@ -23,10 +19,6 @@ UDataLayerFactory::UDataLayerFactory(const FObjectInitializer& ObjectInitializer
 UObject* UDataLayerFactory::FactoryCreateNew(UClass* Class, UObject* InParent, FName InName, EObjectFlags Flags, UObject* Context, FFeedbackContext* Warn)
 {
 	UDataLayerAsset* DataLayerAsset = NewObject<UDataLayerAsset>(InParent, InName, Flags);
-	DataLayerAsset->SetDebugColor(FColor::MakeRandomSeededColor(GetTypeHash(DataLayerAsset->GetFullName())));
-	if (!DataLayerAsset->IsPrivate())
-	{
-		DataLayerAsset->SetType(EDataLayerType::Runtime);
-	}
+	DataLayerAsset->OnCreated();
 	return DataLayerAsset;
 }
\ No newline at end of file
