External Data Layer V1

diff --git a/Engine/Source/Editor/DataLayerEditor/Private/DataLayer/DataLayerHierarchy.cpp b/Engine/Source/Editor/DataLayerEditor/Private/DataLayer/DataLayerHierarchy.cpp
--- a/Engine/Source/Editor/DataLayerEditor/Private/DataLayer/DataLayerHierarchy.cpp
+++ b/Engine/Source/Editor/DataLayerEditor/Private/DataLayer/DataLayerHierarchy.cpp
@@ -298,7 +298,7 @@ void FDataLayerHierarchy::CreateItems(TArray<FSceneOutlinerTreeItemPtr>& OutItem
 						{
 							if (ActorDescInstance != nullptr && !LoadedActors.Contains(ActorDescInstance->GetGuid()) && FActorDescTreeItem::ShouldDisplayInOutliner(ActorDescInstance))
 							{
-								for (const FName& DataLayerInstanceName : ActorDescInstance->GetDataLayerInstanceNames())
+								for (const FName& DataLayerInstanceName : ActorDescInstance->GetDataLayerInstanceNames().ToArray())
 								{
 									if (const UDataLayerInstance* DataLayerInstance = DataLayerManager->GetDataLayerInstance(DataLayerInstanceName))
 									{
@@ -590,7 +590,7 @@ void FDataLayerHierarchy::OnActorDescInstanceAdded(FWorldPartitionActorDescInsta
 	}
 
 	const UDataLayerManager* const DataLayerManager = UDataLayerManager::GetDataLayerManager(GetOwningWorld());
-	const TArray<FName>& DataLayerInstanceNames = InActorDescInstance->GetDataLayerInstanceNames();
+	const FDataLayerInstanceNames& DataLayerInstanceNames = InActorDescInstance->GetDataLayerInstanceNames();
 
 	if (DataLayerManager && DataLayerInstanceNames.Num() > 0)
 	{
@@ -598,7 +598,7 @@ void FDataLayerHierarchy::OnActorDescInstanceAdded(FWorldPartitionActorDescInsta
 		EventData.Type = FSceneOutlinerHierarchyChangedData::Added;
 
 		EventData.Items.Reserve(DataLayerInstanceNames.Num());
-		for (const FName& DataLayerInstanceName : DataLayerInstanceNames)
+		for (const FName& DataLayerInstanceName : DataLayerInstanceNames.ToArray())
 		{
 			const UDataLayerInstance* const DataLayerInstance = DataLayerManager->GetDataLayerInstance(DataLayerInstanceName);
 			EventData.Items.Add(Mode->CreateItemFor<FDataLayerActorDescTreeItem>(FDataLayerActorDescTreeItemData(InActorDescInstance->GetGuid(), InActorDescInstance->GetContainerInstance(), const_cast<UDataLayerInstance*>(DataLayerInstance))));
@@ -615,7 +615,7 @@ void FDataLayerHierarchy::OnActorDescInstanceRemoved(FWorldPartitionActorDescIns
 	}
 
 	const UDataLayerManager* const DataLayerManager = UDataLayerManager::GetDataLayerManager(GetOwningWorld());
-	const TArray<FName>& DataLayerInstanceNames = InActorDescInstance->GetDataLayerInstanceNames();
+	const FDataLayerInstanceNames& DataLayerInstanceNames = InActorDescInstance->GetDataLayerInstanceNames();
 	
 	if (DataLayerManager && DataLayerInstanceNames.Num() > 0)
 	{
@@ -623,7 +623,7 @@ void FDataLayerHierarchy::OnActorDescInstanceRemoved(FWorldPartitionActorDescIns
 		EventData.Type = FSceneOutlinerHierarchyChangedData::Removed;
 		EventData.ItemIDs.Reserve(DataLayerInstanceNames.Num());
 
-		for (const FName& DataLayerInstanceName : DataLayerInstanceNames)
+		for (const FName& DataLayerInstanceName : DataLayerInstanceNames.ToArray())
 		{
 			const UDataLayerInstance* const DataLayer = DataLayerManager->GetDataLayerInstance(DataLayerInstanceName);
 			EventData.ItemIDs.Add(FDataLayerActorDescTreeItem::ComputeTreeItemID(InActorDescInstance->GetGuid(), InActorDescInstance->GetContainerInstance(), DataLayer));
diff --git a/Engine/Source/Editor/DataLayerEditor/Private/DataLayer/DataLayerMode.cpp b/Engine/Source/Editor/DataLayerEditor/Private/DataLayer/DataLayerMode.cpp
--- a/Engine/Source/Editor/DataLayerEditor/Private/DataLayer/DataLayerMode.cpp
+++ b/Engine/Source/Editor/DataLayerEditor/Private/DataLayer/DataLayerMode.cpp
@@ -78,6 +78,7 @@
 #include "ToolMenuEntry.h"
 #include "ToolMenuSection.h"
 #include "ToolMenus.h"
+#include "Settings/EditorExperimentalSettings.h"
 #include "UObject/Class.h"
 #include "UObject/NameTypes.h"
 #include "UObject/ObjectPtr.h"
@@ -87,12 +88,17 @@
 #include "Widgets/DeclarativeSyntaxSupport.h"
 #include "Widgets/SBoxPanel.h"
 #include "Widgets/Views/STreeView.h"
+#include "FolderTreeItem.h"
 #include "WorldDataLayersTreeItem.h"
 #include "WorldPartition/DataLayer/DataLayerAsset.h"
 #include "WorldPartition/DataLayer/DataLayerInstance.h"
 #include "WorldPartition/DataLayer/DataLayerInstanceWithAsset.h"
 #include "WorldPartition/DataLayer/DataLayerManager.h"
+#include "WorldPartition/DataLayer/ExternalDataLayerAsset.h"
+#include "WorldPartition/DataLayer/ExternalDataLayerInstance.h"
+#include "WorldPartition/DataLayer/ExternalDataLayerManager.h"
 #include "WorldPartition/DataLayer/WorldDataLayers.h"
+#include "WorldPartition/WorldPartition.h"
 #include "WorldPartition/WorldPartitionEditorPerProjectUserSettings.h"
 #include "WorldTreeItem.h"
 
@@ -291,7 +297,7 @@ bool FDataLayerMode::CanRenameItem(const ISceneOutlinerTreeItem& Item) const
 	if (Item.IsValid() && (Item.IsA<FDataLayerTreeItem>()))
 	{
 		FDataLayerTreeItem* DataLayerTreeItem = (FDataLayerTreeItem*)&Item;
-		return !DataLayerTreeItem->GetDataLayer()->IsLocked() && DataLayerTreeItem->GetDataLayer()->CanEditDataLayerShortName();
+		return !DataLayerTreeItem->GetDataLayer()->IsReadOnly() && DataLayerTreeItem->GetDataLayer()->CanEditDataLayerShortName();
 	}
 	return false;
 }
@@ -417,14 +423,9 @@ void FDataLayerMode::OnItemDoubleClick(FSceneOutlinerTreeItemPtr Item)
 					{
 						const UDataLayerManager* DataLayerManger = UDataLayerManager::GetDataLayerManager(GetOwningWorld());
 						const TArray<UDataLayerInstance*> ActorEditorContextDataLayers = DataLayerManger->GetActorEditorContextDataLayers();
-						const bool bRemoveCurrent = ActorEditorContextDataLayers.Num() == 1;
-
 						for (UDataLayerInstance* DataLayerInstanceIt : ActorEditorContextDataLayers)
 						{
-							if (DataLayerInstanceIt != DataLayerInstance || bRemoveCurrent)
-							{
-								UDataLayerEditorSubsystem::Get()->RemoveFromActorEditorContext(DataLayerInstanceIt);
-							}
+							UDataLayerEditorSubsystem::Get()->RemoveFromActorEditorContext(DataLayerInstanceIt);
 						}
 					}
 					else
@@ -460,7 +461,7 @@ void FDataLayerMode::DeleteItems(const TArray<TWeakPtr<ISceneOutlinerTreeItem>>&
 		{
 			UDataLayerInstance* DataLayerInstance = DataLayerActorItem->GetDataLayer();
 			AActor* Actor = DataLayerActorItem->GetActor();
-			if (DataLayerInstance && !DataLayerInstance->IsLocked() && Actor)
+			if (DataLayerInstance && !DataLayerInstance->IsReadOnly() && Actor)
 			{
 				ActorsToRemoveFromDataLayer.FindOrAdd(DataLayerInstance).Add(Actor);
 			}
@@ -469,7 +470,7 @@ void FDataLayerMode::DeleteItems(const TArray<TWeakPtr<ISceneOutlinerTreeItem>>&
 		{
 			if (UDataLayerInstance* DataLayerInstance = DataLayerItem->GetDataLayer())
 			{
-				if (!DataLayerInstance->IsLocked())
+				if (!DataLayerInstance->IsReadOnly())
 				{
 					DataLayersToDelete.Add(DataLayerInstance);
 				}
@@ -493,15 +494,24 @@ void FDataLayerMode::DeleteItems(const TArray<TWeakPtr<ISceneOutlinerTreeItem>>&
 
 void FDataLayerMode::DeleteDataLayers(const TArray<UDataLayerInstance*>& InDataLayersToDelete)
 {
+	TArray<UDataLayerInstance*> DataLayersToDelete;
+	for (UDataLayerInstance* DataLayerInstance : InDataLayersToDelete)
+	{
+		if (DataLayerInstance->CanBeRemoved())
+		{
+			DataLayersToDelete.Add(DataLayerInstance);
+		}
+	}
+
 	int32 PrevDeleteCount = SelectedDataLayersSet.Num();
-	for (UDataLayerInstance* DataLayerToDelete : InDataLayersToDelete)
+	for (UDataLayerInstance* DataLayerToDelete : DataLayersToDelete)
 	{
 		SelectedDataLayersSet.Remove(DataLayerToDelete);
 	}
 
 	{
 		const FScopedTransaction Transaction(LOCTEXT("DeleteSelectedDataLayers", "Delete Selected Data Layers"));
-		DataLayerEditorSubsystem->DeleteDataLayers(InDataLayersToDelete);
+		DataLayerEditorSubsystem->DeleteDataLayers(DataLayersToDelete);
 	}
 
 	if ((SelectedDataLayersSet.Num() != PrevDeleteCount) && DataLayerBrowser)
@@ -560,38 +570,49 @@ bool FDataLayerMode::ParseDragDrop(FSceneOutlinerDragDropPayload& OutPayload, co
 		|| !GetDataLayerAssetsFromOperation(Operation, true).IsEmpty();
 }
 
-FSceneOutlinerDragValidationInfo FDataLayerMode::ValidateActorDrop(const ISceneOutlinerTreeItem& DropTarget, bool bMoveOperation /*= false*/) const
+FSceneOutlinerDragValidationInfo FDataLayerMode::ValidateActorDrop(const ISceneOutlinerTreeItem& DropTarget, TArray<AActor*> PayloadActors, bool bMoveOperation /*= false*/) const
 {
-	if (const FDataLayerTreeItem* DataLayerItem = DropTarget.CastTo<FDataLayerTreeItem>())
+	// Only support dropping actors on Data Layer Instances
+	const FDataLayerTreeItem* DataLayerItem = DropTarget.CastTo<FDataLayerTreeItem>();
+	const UDataLayerInstance* DropTargetDataLayerInstance = DataLayerItem ? GetDataLayerInstanceFromTreeItem(DropTarget) : nullptr;
+	if (!DropTargetDataLayerInstance)
 	{
-		const UDataLayerInstance* TargetDataLayer = DataLayerItem->GetDataLayer();
-		if (!TargetDataLayer)
+		return FSceneOutlinerDragValidationInfo(ESceneOutlinerDropCompatibility::IncompatibleGeneric, FText());
+	}
+
+	// Check if user can't add actors to the target Data Layer Instance
+	{
+		FText Reason;
+		if (!DropTargetDataLayerInstance->CanUserAddActors(&Reason))
 		{
-			return FSceneOutlinerDragValidationInfo(ESceneOutlinerDropCompatibility::IncompatibleGeneric, FText());
+			const FText Text = FText::Format(LOCTEXT("CantMoveOrAssignActorsToLockedDataLayer", "Can't {0} actor(s) to Data Layer \"{1}\" : {2}"), bMoveOperation ? FText::FromString(TEXT("move")) : FText::FromString(TEXT("assign")), FText::FromString(DropTargetDataLayerInstance->GetDataLayerShortName()), Reason);
+			return FSceneOutlinerDragValidationInfo(ESceneOutlinerDropCompatibility::IncompatibleGeneric, Text);
 		}
+	}
 
-		if (TargetDataLayer->IsLocked())
+	// Check if can't add one of the payload actors to the target Data Layer Instance
+	for (AActor* Actor : PayloadActors)
+	{
+		FText Reason;
+		if (!DropTargetDataLayerInstance->CanAddActor(Actor, &Reason))
 		{
-			const FText Text = bMoveOperation ? LOCTEXT("CantMoveActorsToLockedDataLayer", "Can't move actors to locked Data Layer") : LOCTEXT("CantAssignActorsToLockedDataLayer", "Can't assign actors to locked Data Layer");
+			const FText Text = FText::Format(LOCTEXT("CantMoveOrAssignActorsToDataLayer", "Can't {0} actor(s) to Data Layer \"{1}\" : {2}"), bMoveOperation ? FText::FromString(TEXT("move")) : FText::FromString(TEXT("assign")), FText::FromString(DropTargetDataLayerInstance->GetDataLayerShortName()), Reason);
 			return FSceneOutlinerDragValidationInfo(ESceneOutlinerDropCompatibility::IncompatibleGeneric, Text);
 		}
+	}
 
-		if (!bMoveOperation)
+	if (!bMoveOperation)
+	{
+		if (GetSelectedDataLayers(SceneOutliner).Num() > 1)
 		{
-			if (GetSelectedDataLayers(SceneOutliner).Num() > 1)
+			if (SceneOutliner->GetTree().IsItemSelected(const_cast<ISceneOutlinerTreeItem&>(DropTarget).AsShared()))
 			{
-				if (SceneOutliner->GetTree().IsItemSelected(const_cast<ISceneOutlinerTreeItem&>(DropTarget).AsShared()))
-				{
-					return FSceneOutlinerDragValidationInfo(ESceneOutlinerDropCompatibility::Compatible, LOCTEXT("AssignToDataLayers", "Assign to Selected Data Layers"));
-				}
+				return FSceneOutlinerDragValidationInfo(ESceneOutlinerDropCompatibility::Compatible, LOCTEXT("AssignToDataLayers", "Assign to Selected Data Layers"));
 			}
 		}
-
-		const FText Text = bMoveOperation ? FText::Format(LOCTEXT("MoveToDataLayer", "Move to Data Layer \"{0}\""), FText::FromString(TargetDataLayer->GetDataLayerShortName())) : FText::Format(LOCTEXT("AssignToDataLayer", "Assign to Data Layer \"{0}\""), FText::FromString(TargetDataLayer->GetDataLayerShortName()));
-		return FSceneOutlinerDragValidationInfo(ESceneOutlinerDropCompatibility::Compatible, Text);
 	}
-
-	return FSceneOutlinerDragValidationInfo::Invalid();
+	const FText Text = FText::Format(LOCTEXT("MoveOrAssignToDataLayer", "{0} to Data Layer \"{1}\""), bMoveOperation ? FText::FromString(TEXT("Move")) : FText::FromString(TEXT("Assign")), FText::FromString(DropTargetDataLayerInstance->GetDataLayerShortName()));
+	return FSceneOutlinerDragValidationInfo(ESceneOutlinerDropCompatibility::Compatible, Text);
 }
 
 FSceneOutlinerDragValidationInfo FDataLayerMode::ValidateDrop(const ISceneOutlinerTreeItem& DropTarget, const FSceneOutlinerDragDropPayload& Payload) const
@@ -600,32 +621,34 @@ FSceneOutlinerDragValidationInfo FDataLayerMode::ValidateDrop(const ISceneOutlin
 	if (!PayloadActors.IsEmpty()) // Adding actor(s) in data layer(s)
 	{
 		const FDataLayerTreeItem* DataLayerItem = DropTarget.CastTo<FDataLayerTreeItem>();
-		const UDataLayerInstance* TargetDataLayer = DataLayerItem ? DataLayerItem->GetDataLayer() : nullptr;
-		const bool bIsPrivateDataLayer = TargetDataLayer && TargetDataLayer->GetAsset() && TargetDataLayer->GetAsset()->IsPrivate();
-		const AWorldDataLayers* OuterWorldDataLayers = TargetDataLayer ? TargetDataLayer->GetOuterWorldDataLayers() : nullptr;
-
-		for (AActor* Actor : PayloadActors)
+		if (UDataLayerInstance* TargetDataLayerInstance = DataLayerItem ? DataLayerItem->GetDataLayer() : nullptr)
 		{
-			if (!Actor->IsUserManaged() || !DataLayerEditorSubsystem->IsActorValidForDataLayer(Actor))
-			{
-				return FSceneOutlinerDragValidationInfo(ESceneOutlinerDropCompatibility::IncompatibleGeneric, LOCTEXT("ActorCantBeAssignedToDataLayer", "Can't assign actors to Data Layer"));
-			}
-			else if (bIsPrivateDataLayer && (Actor->GetLevel()->GetWorldDataLayers() != OuterWorldDataLayers))
+			const bool bIsPrivateDataLayer = TargetDataLayerInstance->GetAsset() && TargetDataLayerInstance->GetAsset()->IsPrivate();
+			const AWorldDataLayers* OuterWorldDataLayers = TargetDataLayerInstance->GetDirectOuterWorldDataLayers();
+
+			for (AActor* Actor : PayloadActors)
 			{
-				return FSceneOutlinerDragValidationInfo(ESceneOutlinerDropCompatibility::IncompatibleGeneric, LOCTEXT("ActorCantBeAssignedToPrivateDataLayerOfOtherWorldDataLayers", "Can't assign actors to Private Data Layer of another WorldDataLayers actor"));
+				if (!Actor->IsUserManaged() || !DataLayerEditorSubsystem->IsActorValidForDataLayerInstances(Actor, { TargetDataLayerInstance }))
+				{
+					return FSceneOutlinerDragValidationInfo(ESceneOutlinerDropCompatibility::IncompatibleGeneric, LOCTEXT("ActorCantBeAssignedToDataLayer", "Can't assign actors to Data Layer"));
+				}
+				else if (bIsPrivateDataLayer && (Actor->GetLevel()->GetWorldDataLayers() != OuterWorldDataLayers))
+				{
+					return FSceneOutlinerDragValidationInfo(ESceneOutlinerDropCompatibility::IncompatibleGeneric, LOCTEXT("ActorCantBeAssignedToPrivateDataLayerOfOtherWorldDataLayers", "Can't assign actors to Private Data Layer of another WorldDataLayers actor"));
+				}
 			}
 		}
 
-		return ValidateActorDrop(DropTarget);
+		return ValidateActorDrop(DropTarget, PayloadActors);
 	}
 	
-	// Drag and drop actor(s) into a Data Layer
+	// Drag and drop actor(s) into a Data Layer Instance
 	if (!GetDataLayerActorPairsFromOperation(Payload.SourceOperation).IsEmpty())
 	{
-		return ValidateActorDrop(DropTarget, true);
+		return ValidateActorDrop(DropTarget, PayloadActors, true);
 	}
 	
-	// Drag and drop DataLayer asset(s)
+	// Drag and drop Data Layer Asset(s)
 	{
 		TArray<const UDataLayerAsset*> DataLayerAssets = GetDataLayerAssetsFromOperation(Payload.SourceOperation);
 		if (!DataLayerAssets.IsEmpty())
@@ -634,7 +657,7 @@ FSceneOutlinerDragValidationInfo FDataLayerMode::ValidateDrop(const ISceneOutlin
 		}
 	}
 	
-	// Drag and drop DataLayerInstance(s)
+	// Drag and drop Data Layer Instance(s)
 	TArray<UDataLayerInstance*> PayloadDataLayers = GetDataLayerInstancesFromOperation(Payload.SourceOperation);
 	if (!PayloadDataLayers.IsEmpty())
 	{
@@ -644,9 +667,9 @@ FSceneOutlinerDragValidationInfo FDataLayerMode::ValidateDrop(const ISceneOutlin
 		{
 			if (!PayloadWorldDatalayers)
 			{
-				PayloadWorldDatalayers = DataLayerInstance->GetOuterWorldDataLayers();
+				PayloadWorldDatalayers = DataLayerInstance->GetDirectOuterWorldDataLayers();
 			}
-			else if (PayloadWorldDatalayers != DataLayerInstance->GetOuterWorldDataLayers())
+			else if (PayloadWorldDatalayers != DataLayerInstance->GetDirectOuterWorldDataLayers())
 			{
 				PayloadWorldDatalayers = nullptr;
 				break;
@@ -658,7 +681,7 @@ FSceneOutlinerDragValidationInfo FDataLayerMode::ValidateDrop(const ISceneOutlin
 		}
 
 		// Drag and drop DataLayerInstanceWithAsset(s) between AWorldDataLayers
-		AWorldDataLayers* DropTargetWorldDataLayers = GetWorldDataLayersFromTreeItem(DropTarget, true);
+		AWorldDataLayers* DropTargetWorldDataLayers = GetWorldDataLayersFromTreeItem(DropTarget);
 		if (DropTargetWorldDataLayers && (PayloadWorldDatalayers != DropTargetWorldDataLayers))
 		{
 			TSet<const UDataLayerAsset*> DataLayerAssets;
@@ -679,39 +702,35 @@ FSceneOutlinerDragValidationInfo FDataLayerMode::ValidateDrop(const ISceneOutlin
 			return FSceneOutlinerDragValidationInfo(ESceneOutlinerDropCompatibility::IncompatibleGeneric, LOCTEXT("NoValidDataLayersToMove", "No valid Data Layer to move"));
 		}
 
-		const FDataLayerTreeItem* DataLayerItem = DropTarget.CastTo<FDataLayerTreeItem>();
-		const FDataLayerActorTreeItem* DataLayerActorTreeItem = DropTarget.CastTo<FDataLayerActorTreeItem>();
-		const UDataLayerInstance* ParentDataLayer = DataLayerItem ? DataLayerItem->GetDataLayer() : nullptr;
-		if (!ParentDataLayer && DataLayerActorTreeItem)
-		{
-			ParentDataLayer = DataLayerActorTreeItem->GetDataLayer();
-		}
-
+		const UDataLayerInstance* ParentDataLayerInstance = GetDataLayerInstanceFromTreeItem(DropTarget);
 		for (UDataLayerInstance* DataLayerInstance : PayloadDataLayers)
 		{
-			if (DataLayerInstance->IsLocked())
 			{
-				return FSceneOutlinerDragValidationInfo(ESceneOutlinerDropCompatibility::IncompatibleGeneric, LOCTEXT("CantMoveLockedDataLayer", "Can't move locked Data Layer"));
+				FText Reason;
+				if (DataLayerInstance->IsReadOnly(&Reason))
+				{
+					return FSceneOutlinerDragValidationInfo(ESceneOutlinerDropCompatibility::IncompatibleGeneric, FText::Format(LOCTEXT("CantMoveReadOnlyDataLayer", "Can't move Data Layer : {0}"), Reason));
+				}
 			}
 
-			if (ParentDataLayer != nullptr)
+			if (ParentDataLayerInstance != nullptr)
 			{
-				FText FailureReason;
-				if (!DataLayerInstance->CanBeChildOf(ParentDataLayer, &FailureReason))
+				FText Reason;
+				if (!DataLayerInstance->CanBeChildOf(ParentDataLayerInstance, &Reason))
 				{
-					return FSceneOutlinerDragValidationInfo(ESceneOutlinerDropCompatibility::IncompatibleGeneric, FailureReason);
+					return FSceneOutlinerDragValidationInfo(ESceneOutlinerDropCompatibility::IncompatibleGeneric, Reason);
 				}
-
-				if (ParentDataLayer->IsLocked())
+				
+				if (ParentDataLayerInstance->IsReadOnly(&Reason))
 				{
-					return FSceneOutlinerDragValidationInfo(ESceneOutlinerDropCompatibility::IncompatibleGeneric, LOCTEXT("CantMoveDataLayerToLockedDataLayer", "Can't move Data Layer to locked Data Layer"));
+					return FSceneOutlinerDragValidationInfo(ESceneOutlinerDropCompatibility::IncompatibleGeneric, FText::Format(LOCTEXT("CantMoveDataLayerToLockedDataLayer", "Can't move to Data Layer : {0}"), Reason));
 				}
 			}
 		}
 
-		if (ParentDataLayer != nullptr)
+		if (ParentDataLayerInstance)
 		{
-			return FSceneOutlinerDragValidationInfo(ESceneOutlinerDropCompatibility::Compatible, FText::Format(LOCTEXT("MoveDataLayerToDataLayer", "Move to Data Layer \"{0}\""), FText::FromString(ParentDataLayer->GetDataLayerShortName())));
+			return FSceneOutlinerDragValidationInfo(ESceneOutlinerDropCompatibility::Compatible, FText::Format(LOCTEXT("MoveDataLayerToDataLayer", "Move to Data Layer \"{0}\""), FText::FromString(ParentDataLayerInstance->GetDataLayerShortName())));
 		}
 		else
 		{
@@ -722,20 +741,37 @@ FSceneOutlinerDragValidationInfo FDataLayerMode::ValidateDrop(const ISceneOutlin
 	return FSceneOutlinerDragValidationInfo::Invalid();
 }
 
-AWorldDataLayers* FDataLayerMode::GetWorldDataLayersFromTreeItem(const ISceneOutlinerTreeItem& TreeItem, bool bDefaultToOwningWorldIfNotFound) const
+UDataLayerInstance* FDataLayerMode::GetDataLayerInstanceFromTreeItem(const ISceneOutlinerTreeItem& TreeItem) const
+{
+	const FDataLayerActorTreeItem* DataLayerActorTargetItem = TreeItem.CastTo<FDataLayerActorTreeItem>();
+	const FDataLayerTreeItem* DataLayerTargetItem = TreeItem.CastTo<FDataLayerTreeItem>();
+	UDataLayerInstance* DataLayerTarget = DataLayerTargetItem ? DataLayerTargetItem->GetDataLayer() : (DataLayerActorTargetItem ? DataLayerActorTargetItem->GetDataLayer() : nullptr);
+	return DataLayerTarget;
+}
+
+AWorldDataLayers* FDataLayerMode::GetWorldDataLayersFromTreeItem(const ISceneOutlinerTreeItem& TreeItem) const
 {
-	AWorldDataLayers* WorldDataLayers = bDefaultToOwningWorldIfNotFound ? GetOwningWorldAWorldDataLayers() : nullptr;
+	AWorldDataLayers* WorldDataLayers = GetOwningWorldAWorldDataLayers();
+
 	if (const FWorldDataLayersTreeItem* WorldDataLayerTreeItem = TreeItem.CastTo<FWorldDataLayersTreeItem>())
 	{
 		WorldDataLayers = WorldDataLayerTreeItem->GetWorldDataLayers();
 	}
-	else if (const FDataLayerTreeItem* DataLayerTargetItem = TreeItem.CastTo<FDataLayerTreeItem>())
+	else if (const FFolderTreeItem* FolderTreeItem = TreeItem.CastTo<FFolderTreeItem>())
 	{
-		if (const UDataLayerInstance* DataLayerTarget = DataLayerTargetItem->GetDataLayer())
+		if (const UObject* RootObject = FolderTreeItem->GetFolder().GetRootObjectPtr())
 		{
-			WorldDataLayers = DataLayerTarget->GetOuterWorldDataLayers();
+			if (AWorldDataLayers* RootObjectWorldDataLayers = RootObject->GetWorld() ? RootObject->GetWorld()->GetWorldDataLayers() : nullptr)
+			{
+				WorldDataLayers = RootObjectWorldDataLayers;
+			}
 		}
 	}
+	else if (const UDataLayerInstance* DataLayerTarget = GetDataLayerInstanceFromTreeItem(TreeItem))
+	{
+		WorldDataLayers = DataLayerTarget->GetDirectOuterWorldDataLayers();
+	}
+	
 	return WorldDataLayers;
 }
 
@@ -749,39 +785,105 @@ FSceneOutlinerDragValidationInfo FDataLayerMode::ValidateDataLayerAssetDrop(cons
 			FText::FromString(GetOwningWorld()->GetName())));
 	}
 
-	const FDataLayerTreeItem* DataLayerTargetItem = DropTarget.CastTo<FDataLayerTreeItem>();
-	const UDataLayerInstance* DataLayerTarget = DataLayerTargetItem ? DataLayerTargetItem->GetDataLayer() : nullptr;
-	const UDataLayerInstanceWithAsset* DataLayerWithAssetTarget = DataLayerTarget ? CastChecked<UDataLayerInstanceWithAsset>(DataLayerTarget) : nullptr;
-
-	if (!DropTarget.CanInteract() || (DataLayerWithAssetTarget != nullptr && DataLayerWithAssetTarget->IsReadOnly()))
+	const AWorldDataLayers* DropTargetWorldDataLayers = GetWorldDataLayersFromTreeItem(DropTarget);
+	const UDataLayerInstance* DropTargetDataLayerInstance = GetDataLayerInstanceFromTreeItem(DropTarget);
+	const UDataLayerInstanceWithAsset* DropTargetDataLayerWithAsset = DropTargetDataLayerInstance ? Cast<UDataLayerInstanceWithAsset>(DropTargetDataLayerInstance) : nullptr;
+	if (!DropTarget.CanInteract() || (DropTargetDataLayerWithAsset && DropTargetDataLayerWithAsset->IsReadOnly()))
 	{
 		return FSceneOutlinerDragValidationInfo(ESceneOutlinerDropCompatibility::Incompatible, FText::Format(LOCTEXT("DropTargetIsReadOnly", "\"{0}\" is read only."), FText::FromString(DropTarget.GetDisplayString())));
 	}
 
-	AWorldDataLayers* DropTargetWorldDataLayers = GetWorldDataLayersFromTreeItem(DropTarget, true);
-	
-	TArray<const UDataLayerInstance*> ExistingDataLayerInstances = DropTargetWorldDataLayers->GetDataLayerInstances(DataLayerAssetsToDrop);
+	// Check for duplicate Data Layer Asset using Outer World's AWorldDataLayers
+	UWorld* OuterWorld = DropTargetWorldDataLayers->GetTypedOuter<UWorld>();
+	const AWorldDataLayers* OuterWorldDataLayers = OuterWorld->GetWorldDataLayers();
+	const TArray<const UDataLayerInstance*> ExistingDataLayerInstances = OuterWorldDataLayers->GetDataLayerInstances(DataLayerAssetsToDrop);
+	if (!ExistingDataLayerInstances.IsEmpty())
+	{
+		TArray<FString> ExistingAssestNames;
+		Algo::Transform(ExistingDataLayerInstances, ExistingAssestNames, [](const UDataLayerInstance* DataLayerInstance) { return Cast<UDataLayerInstanceWithAsset>(DataLayerInstance)->GetAsset()->GetName(); });
+		FString ExistingAssetNamesString = FString::Join(ExistingAssestNames, TEXT(","));
+		return FSceneOutlinerDragValidationInfo(ESceneOutlinerDropCompatibility::Incompatible, FText::Format(LOCTEXT("CantCreateDataLayerInstanceSameAsset", "Cannot create Data Layer Instance(s), there are already Data Layer Instance(s) for \"{0}\"."), FText::FromString(ExistingAssetNamesString)));
+	}
+
+	// Check if External Data Layer Asset can be added
+	UExternalDataLayerManager* ExternalDataLayerManager = UExternalDataLayerManager::GetExternalDataLayerManager(OuterWorld);
+	if (!ExternalDataLayerManager)
+	{
+		return FSceneOutlinerDragValidationInfo(ESceneOutlinerDropCompatibility::Incompatible, LOCTEXT("CantCreateDataLayerInstanceWithExternalDataLayerAsset", "Cannot create External Data Layer Instance : Partitioned world doesn't support External Data Layers."));
+	}
+	if (!GetDefault<UEditorExperimentalSettings>()->bEnableWorldPartitionExternalDataLayers)
+	{
+		return FSceneOutlinerDragValidationInfo(ESceneOutlinerDropCompatibility::Incompatible, LOCTEXT("CantCreateExternalDataLayerInstanceExperimentalFlagDisabled", "Cannot create External Data Layer Instance : Experimental flag 'Enable World Partition External Data Layers' is disabled."));
+	}
+	FText InjectionFailureReason;
+	if (Algo::AnyOf(DataLayerAssetsToDrop, [ExternalDataLayerManager, &InjectionFailureReason](const UDataLayerAsset* DataLayerAsset) { return DataLayerAsset && DataLayerAsset->IsA<UExternalDataLayerAsset>() && !ExternalDataLayerManager->CanInjectExternalDataLayerAsset(Cast<UExternalDataLayerAsset>(DataLayerAsset), &InjectionFailureReason); }))
+	{
+		return FSceneOutlinerDragValidationInfo(ESceneOutlinerDropCompatibility::Incompatible, FText::Format(LOCTEXT("CantCreateDataLayerInstanceWithExternalDataLayerAsset", "Cannot create External Data Layer Instance : {0}"), InjectionFailureReason));
+	}
 
-	if (ExistingDataLayerInstances.IsEmpty())
+	if (Algo::AnyOf(DataLayerAssetsToDrop, [](const UDataLayerAsset* DataLayerAsset) { return DataLayerAsset && DataLayerAsset->IsA<UExternalDataLayerAsset>(); }))
 	{
-		if (DropTarget.CastTo<FDataLayerTreeItem>() || DropTarget.CastTo<FWorldDataLayersTreeItem>())
+		if (DropTarget.IsA<FDataLayerTreeItem>() || DropTarget.IsA<FDataLayerActorTreeItem>())
 		{
-			return FSceneOutlinerDragValidationInfo(ESceneOutlinerDropCompatibility::Compatible, FText::Format(LOCTEXT("CreateAllDataLayerFromAssetDropOnPart", "Create Data Layer Instances under \"{0}\""), FText::FromString(DropTarget.GetDisplayString())));
+			return FSceneOutlinerDragValidationInfo(ESceneOutlinerDropCompatibility::Incompatible, FText::Format(LOCTEXT("CantCreateDataLayerInstanceUnderItem", "Cannot create External Data Layer Instance under {0}"), FText::FromString(DropTarget.GetDisplayString())));
 		}
-		else
+	}
+
+	auto PassesAssetReferenceFiltering = [](const UObject* InReferencingObject, const UDataLayerAsset* InDataLayerAsset, FText* OutReason)
+	{
+		if (InReferencingObject->IsA<AWorldDataLayers>() && InDataLayerAsset->IsA<UExternalDataLayerAsset>())
 		{
-			return FSceneOutlinerDragValidationInfo(ESceneOutlinerDropCompatibility::Compatible, LOCTEXT("CreateAllDataLayerFromAssetDrop", "Create Data Layer Instances"));
+			return true;
+		}
+		FAssetReferenceFilterContext AssetReferenceFilterContext;
+		AssetReferenceFilterContext.ReferencingAssets.Add(FAssetData(InReferencingObject));
+		TSharedPtr<IAssetReferenceFilter> AssetReferenceFilter = GEditor->MakeAssetReferenceFilter(AssetReferenceFilterContext);
+		return AssetReferenceFilter.IsValid() ? AssetReferenceFilter->PassesFilter(FAssetData(InDataLayerAsset), OutReason) : true;
+	};
+
+	// Check if can reference Data Layer Asset
+	const UExternalDataLayerInstance* RootExternalDataLayerInstance = DropTargetDataLayerWithAsset ? DropTargetDataLayerWithAsset->GetRootExternalDataLayerInstance() : nullptr;
+	const UExternalDataLayerAsset* RootExternalDataLayerAsset = RootExternalDataLayerInstance ? RootExternalDataLayerInstance->GetExternalDataLayerAsset() : nullptr;
+	if (const UObject* ReferencingObject = RootExternalDataLayerAsset ? Cast<UObject>(RootExternalDataLayerAsset) : Cast<UObject>(DropTargetWorldDataLayers))
+	{
+		for (const UDataLayerAsset* DataLayerAsset : DataLayerAssetsToDrop)
+		{
+			FText FailureReason;
+			if (!PassesAssetReferenceFiltering(ReferencingObject, DataLayerAsset, &FailureReason))
+			{	
+				return FSceneOutlinerDragValidationInfo(ESceneOutlinerDropCompatibility::Incompatible, FText::Format(LOCTEXT("CantCreateDataLayerInstancePassFilterFailed", "Cannot create Data Layer Instance : {0}"), FailureReason));
+			}
 		}
 	}
-	else
+
+	// Check if target data layer supports having dropped asset as a child
+	if (DropTargetDataLayerWithAsset)
 	{
-		TArray<FString> ExistingAssestNames;
-		Algo::Transform(ExistingDataLayerInstances, ExistingAssestNames, [](const UDataLayerInstance* DataLayerInstance) { return CastChecked<UDataLayerInstanceWithAsset>(DataLayerInstance)->GetAsset()->GetName(); });
-		
-		FString ExistingAssetNamesString = FString::Join(ExistingAssestNames, TEXT(","));
+		EDataLayerType ParentType = DropTargetDataLayerWithAsset->GetType();
+		auto IsParentDataLayerTypeCompatible = [ParentType](const UDataLayerAsset* InChildDataLayerAsset)
+		{
+			EDataLayerType ChildType = InChildDataLayerAsset->GetType();
+			return (ChildType != EDataLayerType::Unknown) &&
+				   (ParentType != EDataLayerType::Unknown) &&
+				   (ParentType == EDataLayerType::Editor || ChildType == EDataLayerType::Runtime);
+		};
 
-		return FSceneOutlinerDragValidationInfo(ESceneOutlinerDropCompatibility::Incompatible, FText::Format(LOCTEXT("CantCreateDataLayerInstanceSameAsset", "Cannot create Data Layer Instance(s) there are already Data Layer Instance(s) for \"{0}\"."),
-			FText::FromString(ExistingAssetNamesString)));
+		for (const UDataLayerAsset* DataLayerAssetToDrop : DataLayerAssetsToDrop)
+		{
+			if (!IsParentDataLayerTypeCompatible(DataLayerAssetToDrop))
+			{
+				return FSceneOutlinerDragValidationInfo(ESceneOutlinerDropCompatibility::Incompatible, FText::Format(LOCTEXT("CantCreateDataLayerInstanceIncompatibleChildType", "Cannot create Data Layer Instance : {0} Data Layer cannot have {1} child Data Layers"), UEnum::GetDisplayValueAsText(ParentType), UEnum::GetDisplayValueAsText(DataLayerAssetToDrop->GetType())));
+			}
+		}
+	}
+	
+	if (DropTarget.CastTo<FDataLayerTreeItem>() || DropTarget.CastTo<FWorldDataLayersTreeItem>())
+	{
+		return FSceneOutlinerDragValidationInfo(ESceneOutlinerDropCompatibility::Compatible, FText::Format(LOCTEXT("CreateAllDataLayerFromAssetDropOnPart", "Create Data Layer Instances under \"{0}\""), FText::FromString(DropTarget.GetDisplayString())));
+	}
+	else
+	{
+		return FSceneOutlinerDragValidationInfo(ESceneOutlinerDropCompatibility::Compatible, LOCTEXT("CreateAllDataLayerFromAssetDrop", "Create Data Layer Instances"));
 	}
 }
 
@@ -924,8 +1026,7 @@ TArray<AActor*> FDataLayerMode::GetActorsFromOperation(const FDragDropOperation&
 
 void FDataLayerMode::OnDrop(ISceneOutlinerTreeItem& DropTarget, const FSceneOutlinerDragDropPayload& Payload, const FSceneOutlinerDragValidationInfo& ValidationInfo) const
 {
-	FDataLayerTreeItem* DataLayerItem = DropTarget.CastTo<FDataLayerTreeItem>();
-	UDataLayerInstance* TargetDataLayer = DataLayerItem ? DataLayerItem->GetDataLayer() : nullptr;
+	UDataLayerInstance* TargetDataLayerInstance = GetDataLayerInstanceFromTreeItem(DropTarget);
 
 	TArray<AActor*> ActorsToAdd = GetActorsFromOperation(Payload.SourceOperation);
 	if (!ActorsToAdd.IsEmpty()) // Adding actor(s) in data layer(s)
@@ -941,10 +1042,10 @@ void FDataLayerMode::OnDrop(ISceneOutlinerTreeItem& DropTarget, const FSceneOutl
 			}
 		}
 
-		if (TargetDataLayer)
+		if (TargetDataLayerInstance)
 		{
 			const FScopedTransaction Transaction(LOCTEXT("DataLayerOutlinerAddActorsToDataLayer", "Add Actors to Data Layer"));
-			DataLayerEditorSubsystem->AddActorsToDataLayer(ActorsToAdd, TargetDataLayer);
+			DataLayerEditorSubsystem->AddActorsToDataLayer(ActorsToAdd, TargetDataLayerInstance);
 		}
 
 		return;
@@ -953,14 +1054,14 @@ void FDataLayerMode::OnDrop(ISceneOutlinerTreeItem& DropTarget, const FSceneOutl
 	if (Payload.SourceOperation.IsOfType<FDataLayerActorMoveOp>()) // Moving actor(s) into a Data Layer
 	{
 		TArray<FDataLayerActorMoveElement> ActorsToMove = GetDataLayerActorPairsFromOperation(Payload.SourceOperation);
-		if (!ActorsToMove.IsEmpty() && TargetDataLayer)
+		if (!ActorsToMove.IsEmpty() && TargetDataLayerInstance)
 		{
 			const FScopedTransaction Transaction(LOCTEXT("DataLayerOutlinerMoveActorsToDataLayer", "Move Actors to Data Layer"));
 			for (const auto& DataLayerActorPair : ActorsToMove)
 			{
 				if (AActor* ActorPtr = DataLayerActorPair.Key.Get())
 				{
-					DataLayerEditorSubsystem->AddActorToDataLayer(ActorPtr, TargetDataLayer);
+					DataLayerEditorSubsystem->AddActorToDataLayer(ActorPtr, TargetDataLayerInstance);
 					DataLayerEditorSubsystem->RemoveActorFromDataLayer(ActorPtr, DataLayerActorPair.Value.Get());
 				}
 			}
@@ -981,8 +1082,8 @@ void FDataLayerMode::OnDrop(ISceneOutlinerTreeItem& DropTarget, const FSceneOutl
 	TArray<UDataLayerInstance*> DataLayerInstances = GetDataLayerInstancesFromOperation(Payload.SourceOperation);
 
 	// Handle dropping DataLayerInstanceWithAsset(s) in a different AWorldDataLayers
-	AWorldDataLayers* DropTargetWorldDataLayers = GetWorldDataLayersFromTreeItem(DropTarget, true);
-	AWorldDataLayers* PayloadWorldDatalayers = DataLayerInstances.Num() ? DataLayerInstances[0]->GetOuterWorldDataLayers() : nullptr; // Validation guarantees that they are all part of the same AWorldDataLayers
+	AWorldDataLayers* DropTargetWorldDataLayers = GetWorldDataLayersFromTreeItem(DropTarget);
+	AWorldDataLayers* PayloadWorldDatalayers = DataLayerInstances.Num() ? DataLayerInstances[0]->GetDirectOuterWorldDataLayers() : nullptr; // Validation guarantees that they are all part of the same AWorldDataLayers
 	if (DropTargetWorldDataLayers && (PayloadWorldDatalayers != DropTargetWorldDataLayers))
 	{
 		TSet<const UDataLayerAsset*> DataLayerAssets;
@@ -1001,7 +1102,7 @@ void FDataLayerMode::OnDrop(ISceneOutlinerTreeItem& DropTarget, const FSceneOutl
 	}
 	
 	// Moving DataLayerInstance(s)
-	SetParentDataLayer(DataLayerInstances, TargetDataLayer);
+	SetParentDataLayer(DataLayerInstances, TargetDataLayerInstance);
 }
 
 void FDataLayerMode::OnDataLayerAssetDropped(const TArray<const UDataLayerAsset*>& DroppedDataLayerAsset, ISceneOutlinerTreeItem& DropTarget) const
@@ -1012,14 +1113,12 @@ void FDataLayerMode::OnDataLayerAssetDropped(const TArray<const UDataLayerAsset*
 	{
 		FDataLayerCreationParameters Params;
 		Params.DataLayerAsset = const_cast<UDataLayerAsset*>(DataLayerAsset);
-		Params.WorldDataLayers = GetWorldDataLayersFromTreeItem(DropTarget, true);
-		UDataLayerInstance* DataLayerInstance = DataLayerEditorSubsystem->CreateDataLayerInstance(Params);
-
-		if (FDataLayerTreeItem* DataLayerItem = DropTarget.CastTo<FDataLayerTreeItem>())
+		Params.WorldDataLayers = GetWorldDataLayersFromTreeItem(DropTarget);
+		if (UDataLayerInstance* DataLayerInstance = DataLayerEditorSubsystem->CreateDataLayerInstance(Params))
 		{
-			if (UDataLayerInstance* ParentDataLayer = DataLayerItem->GetDataLayer())
+			if (UDataLayerInstance* ParentDataLayerInstance = GetDataLayerInstanceFromTreeItem(DropTarget))
 			{
-				DataLayerEditorSubsystem->SetParentDataLayer(DataLayerInstance, ParentDataLayer);
+				DataLayerEditorSubsystem->SetParentDataLayer(DataLayerInstance, ParentDataLayerInstance);
 			}
 		}
 	}
@@ -1136,7 +1235,7 @@ TSharedPtr<FDragDropOperation> FDataLayerMode::CreateDragDropOperation(const FPo
 	if (DraggedObjects.Has<FDataLayerTreeItem>())
 	{
 		TArray<TWeakObjectPtr<UDataLayerInstance>> DataLayerInstances = DraggedObjects.GetData<TWeakObjectPtr<UDataLayerInstance>>(FWeakDataLayerSelector());
-		if (DataLayerInstances.FindByPredicate([&](const TWeakObjectPtr<UDataLayerInstance>& DataLayer) { return DataLayer.IsValid() && DataLayer->IsLocked(); }))
+		if (DataLayerInstances.FindByPredicate([&](const TWeakObjectPtr<UDataLayerInstance>& DataLayerInstance) { return DataLayerInstance.IsValid() && DataLayerInstance->IsReadOnly(); }))
 		{
 			return TSharedPtr<FDragDropOperation>();
 		}
@@ -1214,16 +1313,14 @@ TArray<UDataLayerInstance*> FDataLayerMode::GetSelectedDataLayers(SSceneOutliner
 	return MoveTemp(ValidSelectedDataLayers);
 }
 
-void FDataLayerMode::CreateDataLayerPicker(UToolMenu* InMenu, FOnDataLayerPicked OnDataLayerPicked, bool bInShowRoot /*= false*/)
+void FDataLayerMode::CreateDataLayerPicker(UToolMenu* InMenu, FOnDataLayerInstancePicked OnDataLayerInstancePicked, FOnShouldFilterDataLayerInstance OnShouldFilterDataLayerInstance = FOnShouldFilterDataLayerInstance(), bool bInShowRoot /*= false*/)
 {
+	FToolMenuSection& Section = InMenu->AddSection("DataLayers", LOCTEXT("DataLayers", "Data Layers"));
 	if (bInShowRoot)
 	{
-		FToolMenuSection& Section = InMenu->AddSection("DataLayers", LOCTEXT("DataLayers", "Data Layers"));
-		Section.AddMenuEntry("Root", LOCTEXT("Root", "<Root>"), FText(), FSlateIcon(), FUIAction(FExecuteAction::CreateLambda([=]() { OnDataLayerPicked.ExecuteIfBound(nullptr); })));
+		Section.AddMenuEntry("Root", LOCTEXT("Root", "<Root>"), FText(), FSlateIcon(), FUIAction(FExecuteAction::CreateLambda([=]() { OnDataLayerInstancePicked.ExecuteIfBound(nullptr); })));
 	}
-	
-	FToolMenuSection& Section = InMenu->AddSection(FName(), LOCTEXT("ExistingDataLayers", "Existing Data Layers:"));
-	TSharedRef<SWidget> DataLayerPickerWidget = FDataLayerPickingMode::CreateDataLayerPickerWidget(OnDataLayerPicked);
+	TSharedRef<SWidget> DataLayerPickerWidget = FDataLayerPickingMode::CreateDataLayerPickerWidget(OnDataLayerInstancePicked, OnShouldFilterDataLayerInstance);
 	Section.AddEntry(FToolMenuEntry::InitWidget("DataLayerPickerWidget", DataLayerPickerWidget, FText::GetEmpty(), false));
 }
 
@@ -1296,7 +1393,7 @@ void FDataLayerMode::RegisterContextMenu()
 			FDataLayerMode* Mode = const_cast<FDataLayerMode*>(static_cast<const FDataLayerMode*>(SceneOutliner->GetMode()));
 			check(Mode);
 			TArray<UDataLayerInstance*> SelectedDataLayers = Mode->GetSelectedDataLayers(SceneOutliner);
-			const bool bSelectedDataLayersContainsLocked = Algo::AnyOf(SelectedDataLayers, [](const UDataLayerInstance* DataLayerInstance) { return DataLayerInstance->IsLocked(); });
+			const bool bSelectedDataLayersContainsLocked = Algo::AnyOf(SelectedDataLayers, [](const UDataLayerInstance* DataLayerInstance) { return DataLayerInstance->IsReadOnly(); });
 			const bool bSelectedDataLayersContainsExternalPackage = Algo::AnyOf(SelectedDataLayers, [](const UDataLayerInstance* DataLayerInstance) { return DataLayerInstance->IsPackageExternal(); });
 
 			TArray<AActor*> SelectedActors;
@@ -1331,13 +1428,13 @@ void FDataLayerMode::RegisterContextMenu()
 					if (!BestCandidate)
 					{
 						TSet<AWorldDataLayers*> OuterWorldDataLayers;
-						Algo::TransformIf(Mode->SelectedDataLayersSet, OuterWorldDataLayers, [](const TWeakObjectPtr<const UDataLayerInstance>& DataLayerInstance) { return DataLayerInstance.IsValid() && DataLayerInstance->GetOuterWorldDataLayers(); }, [](const TWeakObjectPtr<const UDataLayerInstance>& DataLayerInstance) { return DataLayerInstance->GetOuterWorldDataLayers(); });
+						Algo::TransformIf(Mode->SelectedDataLayersSet, OuterWorldDataLayers, [](const TWeakObjectPtr<const UDataLayerInstance>& DataLayerInstance) { return DataLayerInstance.IsValid() && DataLayerInstance->GetDirectOuterWorldDataLayers(); }, [](const TWeakObjectPtr<const UDataLayerInstance>& DataLayerInstance) { return DataLayerInstance->GetDirectOuterWorldDataLayers(); });
 						BestCandidate = (OuterWorldDataLayers.Num() == 1) ? *OuterWorldDataLayers.CreateConstIterator() : nullptr;
 					}
 					if (!BestCandidate)
 					{
 						TSet<AWorldDataLayers*> OuterWorldDataLayers;
-						Algo::TransformIf(Mode->SelectedDataLayerActors, OuterWorldDataLayers, [](const FSelectedDataLayerActor& DataLayerActor) { return DataLayerActor.Key.IsValid() && DataLayerActor.Key->GetOuterWorldDataLayers(); }, [](const FSelectedDataLayerActor& DataLayerActor) { return DataLayerActor.Key->GetOuterWorldDataLayers(); });
+						Algo::TransformIf(Mode->SelectedDataLayerActors, OuterWorldDataLayers, [](const FSelectedDataLayerActor& DataLayerActor) { return DataLayerActor.Key.IsValid() && DataLayerActor.Key->GetDirectOuterWorldDataLayers(); }, [](const FSelectedDataLayerActor& DataLayerActor) { return DataLayerActor.Key->GetDirectOuterWorldDataLayers(); });
 						BestCandidate = (OuterWorldDataLayers.Num() == 1) ? *OuterWorldDataLayers.CreateConstIterator() : nullptr;
 					}
 					if (!BestCandidate)
@@ -1347,7 +1444,7 @@ void FDataLayerMode::RegisterContextMenu()
 					return BestCandidate;
 				};
 
-				auto CreateNewDataLayerInternal = [&GetBestCandidateWorldDataLayersFromSelection, SceneOutliner, Mode](UDataLayerInstance* InParentDataLayer, const UDataLayerAsset* InDataLayerAsset = nullptr, bool bInIsPrivate = false) -> UDataLayerInstance*
+				auto CreateNewDataLayerInternal = [&GetBestCandidateWorldDataLayersFromSelection, SceneOutliner, Mode](const UDataLayerAsset* InDataLayerAsset = nullptr, bool bInIsPrivate = false) -> UDataLayerInstance*
 				{
 					const FScopedTransaction Transaction(LOCTEXT("CreateNewDataLayer", "Create New Data Layer"));
 					AWorldDataLayers* TargetWorldDataLayers = GetBestCandidateWorldDataLayersFromSelection(SceneOutliner, Mode);
@@ -1356,19 +1453,15 @@ void FDataLayerMode::RegisterContextMenu()
 
 					FDataLayerCreationParameters CreationParams;
 					CreationParams.DataLayerAsset = const_cast<UDataLayerAsset*>(InDataLayerAsset);
-					CreationParams.WorldDataLayers = InParentDataLayer ? InParentDataLayer->GetOuterWorldDataLayers() : TargetWorldDataLayers;
+					CreationParams.WorldDataLayers = TargetWorldDataLayers;
 					CreationParams.bIsPrivate = bInIsPrivate;
 					if (UDataLayerInstance* NewDataLayerInstance = UDataLayerEditorSubsystem::Get()->CreateDataLayerInstance(CreationParams))
 					{
-						UDataLayerEditorSubsystem::Get()->SetParentDataLayer(NewDataLayerInstance, InParentDataLayer);
-
 						Mode->SelectedDataLayersSet.Add(NewDataLayerInstance);
 						// Select it and open a rename when it gets refreshed
 						SceneOutliner->OnItemAdded(NewDataLayerInstance, SceneOutliner::ENewItemAction::Select | SceneOutliner::ENewItemAction::Rename);
-
 						return NewDataLayerInstance;
 					}
-
 					return nullptr;
 				};
 
@@ -1377,39 +1470,44 @@ void FDataLayerMode::RegisterContextMenu()
 					if (UDataLayerEditorSubsystem::Get()->HasDeprecatedDataLayers())
 					{
 						check(InDataLayerAsset == nullptr);
-						return CreateNewDataLayerInternal(InParentDataLayer, nullptr, false);
+						return CreateNewDataLayerInternal(nullptr, false);
 					}
 					else
 					{
-						return CreateNewDataLayerInternal(InParentDataLayer, InDataLayerAsset, bInIsPrivate);
+						return CreateNewDataLayerInternal(InDataLayerAsset, bInIsPrivate);
 					}
 				};
 
-				auto CreateNewEmptyDataLayer = [SceneOutliner, Mode, CreateNewDataLayerInternal](UDataLayerInstance* InParentDataLayer = nullptr, bool bInIsPrivate = false) -> UDataLayerInstance*
+				auto CreateNewEmptyDataLayer = [SceneOutliner, Mode, CreateNewDataLayerInternal](bool bInIsPrivate = false) -> UDataLayerInstance*
 				{
-					return CreateNewDataLayerInternal(InParentDataLayer, nullptr, bInIsPrivate);
+					return CreateNewDataLayerInternal(nullptr, bInIsPrivate);
 				};
 
 				const AWorldDataLayers* WorldDataLayers = Mode->GetOwningWorld() ? Mode->GetOwningWorld()->GetWorldDataLayers() : nullptr;
 				if (WorldDataLayers && !WorldDataLayers->HasDeprecatedDataLayers())
 				{
 					Section.AddSubMenu("CreateNewDataLayerWithAsset", LOCTEXT("CreateNewDataLayerWithAssetSubMenu", "Create New Data Layer With Asset"), LOCTEXT("CreateNewDataLayerWithAssetSubMenu_ToolTip", "Create New Data Layer With Asset"),
-						FNewToolMenuDelegate::CreateLambda([CreateNewDataLayer, WorldDataLayers](UToolMenu* InSubMenu)
+						FNewToolMenuDelegate::CreateLambda([CreateNewDataLayer, WorldDataLayers, Mode](UToolMenu* InSubMenu)
 						{
 							const bool bAllowClear = false;
 							const bool bAllowCopyPaste = false;
 							const TArray<const UClass*> AllowedClasses = { UDataLayerAsset::StaticClass() };
+							const TArray<const UClass*> NewAssetDisallowedClasses = { UExternalDataLayerAsset::StaticClass() };
 							FToolMenuSection& Section = InSubMenu->AddSection("Data Layer Asset");
 							TSharedRef<SWidget> MenuWidget = PropertyCustomizationHelpers::MakeAssetPickerWithMenu(
 								FAssetData(),
 								bAllowClear,
 								bAllowCopyPaste,
 								AllowedClasses,
-								PropertyCustomizationHelpers::GetNewAssetFactoriesForClasses(AllowedClasses),
-								FOnShouldFilterAsset::CreateLambda([WorldDataLayers](const FAssetData& InAssetData)
+								PropertyCustomizationHelpers::GetNewAssetFactoriesForClasses(AllowedClasses, NewAssetDisallowedClasses),
+								FOnShouldFilterAsset::CreateLambda([Mode, WorldDataLayers](const FAssetData& InAssetData)
 								{
-									UDataLayerAsset* DataLayerAsset = Cast<UDataLayerAsset>(InAssetData.GetAsset());
-									return !DataLayerAsset || WorldDataLayers->GetDataLayerInstance(DataLayerAsset);
+									// Filter already used Data Layers Assets and External Data Layer Assets that can't be added (those already added or not part of a registered GFD action)
+									const UExternalDataLayerManager* ExternalDataLayerManager = UExternalDataLayerManager::GetExternalDataLayerManager(Mode->GetOwningWorld());
+									const UDataLayerAsset* DataLayerAsset = Cast<UDataLayerAsset>(InAssetData.GetAsset());
+									const UExternalDataLayerAsset* ExternalDataLayerAsset = Cast<UExternalDataLayerAsset>(DataLayerAsset);
+									const bool bCanInjectExternalDataLayerAsset = ExternalDataLayerAsset && ExternalDataLayerManager && ExternalDataLayerManager->CanInjectExternalDataLayerAsset(ExternalDataLayerAsset);
+									return !DataLayerAsset || (ExternalDataLayerAsset && !bCanInjectExternalDataLayerAsset) || WorldDataLayers->GetDataLayerInstance(DataLayerAsset);
 								}),
 								FOnAssetSelected::CreateLambda([CreateNewDataLayer](const FAssetData& InAssetData)
 								{
@@ -1427,77 +1525,132 @@ void FDataLayerMode::RegisterContextMenu()
 					Section.AddMenuEntry("CreateNewDataLayer", LOCTEXT("CreateNewDataLayer", "Create New Data Layer"), FText(), FSlateIcon(),
 						FUIAction(FExecuteAction::CreateLambda([CreateNewEmptyDataLayer]() { CreateNewEmptyDataLayer(); })));
 				}
-				
+
 				Section.AddMenuEntry("CreateNewDataLayerPrivate", LOCTEXT("CreateNewDataLayerPrivate", "Create New Private Data Layer"), LOCTEXT("CreateNewDataLayerPrivateToolTip", "Creates an Editor Only Data Layer that cannot be used by other Worlds."), FSlateIcon(),
 					FUIAction(
-						FExecuteAction::CreateLambda([CreateNewEmptyDataLayer]() { CreateNewEmptyDataLayer(nullptr, true); }),
+						FExecuteAction::CreateLambda([CreateNewEmptyDataLayer]() { CreateNewEmptyDataLayer(true); }),
 						FCanExecuteAction(),
 						FIsActionChecked(),
 						FIsActionButtonVisible::CreateLambda([]() { return !UDataLayerEditorSubsystem::Get()->HasDeprecatedDataLayers(); })));
 
+				const bool bAllSelectedDataLayersCanBeRemoved = Algo::AllOf(SelectedDataLayers, [](const UDataLayerInstance* DataLayerInstance) { return DataLayerInstance->CanBeRemoved(); });
+				Section.AddMenuEntry("DeleteSelectedDataLayers", LOCTEXT("DeleteSelectedDataLayers", "Delete Selected Data Layer(s)"), FText(), FSlateIcon(),
+					FUIAction(
+						FExecuteAction::CreateLambda([Mode, SelectedDataLayers]() 
+						{
+							check(!SelectedDataLayers.IsEmpty());
+							Mode->DeleteDataLayers(SelectedDataLayers);
+						}),
+						FCanExecuteAction::CreateLambda([SelectedDataLayers, bSelectedDataLayersContainsLocked, bAllSelectedDataLayersCanBeRemoved]
+						{
+							return !SelectedDataLayers.IsEmpty() && !bSelectedDataLayersContainsLocked && bAllSelectedDataLayersCanBeRemoved;
+						})
+					));
+
+				const bool bSelectedDataLayersSupportParenting = Algo::AllOf(SelectedDataLayers, [](const UDataLayerInstance* DataLayerInstance) { return !DataLayerInstance->IsReadOnly() && DataLayerInstance->CanHaveParentDataLayerInstance(); });
+				if (!SelectedDataLayers.IsEmpty() && !bSelectedDataLayersContainsLocked && bSelectedDataLayersSupportParenting)
+				{
+					Section.AddSubMenu("MoveSelectedDataLayersTo", LOCTEXT("MoveSelectedDataLayersTo", "Move Selected Data Layer(s) To"), FText(),
+						FNewToolMenuDelegate::CreateLambda([Mode, SelectedDataLayers](UToolMenu* InSubMenu)
+						{
+							CreateDataLayerPicker(InSubMenu,
+								FOnDataLayerInstancePicked::CreateLambda([Mode, SelectedDataLayers](UDataLayerInstance* TargetDataLayerInstance)
+								{
+									TArray<UDataLayerInstance*> DataLayerInstances;
+									for (UDataLayerInstance* DataLayerInstance : SelectedDataLayers)
+									{
+										if (ensure(DataLayerInstance->CanBeChildOf(TargetDataLayerInstance)))
+										{
+											DataLayerInstances.Add(DataLayerInstance);
+										}
+									}
+									Mode->SetParentDataLayer(DataLayerInstances, TargetDataLayerInstance);
+								}),
+								FOnShouldFilterDataLayerInstance::CreateLambda([Mode, SelectedDataLayers](const UDataLayerInstance* InCandidateDataLayerInstance)
+								{
+									for (UDataLayerInstance* DataLayerInstance : SelectedDataLayers)
+									{
+										if (!DataLayerInstance->CanBeChildOf(InCandidateDataLayerInstance))
+										{
+											return false;
+										}
+									}
+									return true;
+								}),
+								/*bShowRoot*/true);
+						}));
+				}
+
+				Section.AddMenuEntry("Copy Selected Data Layer Instances(s) File Path", LOCTEXT("CopySelectedDataLayerInstancessFilePath", "Copy Selected Data Layer Instance(s) File Path"), FText(), FSlateIcon(),
+					FUIAction(
+						FExecuteAction::CreateLambda([SelectedDataLayers]()
+						{
+							check(!SelectedDataLayers.IsEmpty());
+							{
+								TArray<const UObject*> Objects;
+								Algo::Transform(SelectedDataLayers, Objects, [](UDataLayerInstance* DataLayerInstance) { return DataLayerInstance; });
+								FExternalPackageHelper::CopyObjectsExternalPackageFilePathToClipboard(Objects);
+							}
+						}),
+						FCanExecuteAction::CreateLambda([bSelectedDataLayersContainsExternalPackage] { return bSelectedDataLayersContainsExternalPackage; })
+					));
+
+				Section.AddSeparator("SectionsSeparator");
+
+
 				const bool bAllSelectedDataLayersCanAddActors = Algo::AllOf(SelectedDataLayers, [](const UDataLayerInstance* DataLayerInstance) { return DataLayerInstance->CanUserAddActors(); });
 				FText AddActorToolTip = bSelectedActorsAreAllUserManaged ? FText::GetEmpty() : SelectionContainsNonUserManagedActorsText;
-				Section.AddMenuEntry("AddSelectedActorsToSelectedDataLayers", LOCTEXT("AddSelectedActorsToSelectedDataLayersMenu", "Add Selected Actors to Selected Data Layers"), AddActorToolTip, FSlateIcon(),
+				Section.AddMenuEntry("AddSelectedActorsToSelectedDataLayers", LOCTEXT("AddSelectedActorsToSelectedDataLayersMenu", "Add Selected Actor(s) to Selected Data Layer(s)"), AddActorToolTip, FSlateIcon(),
 					FUIAction(
-						FExecuteAction::CreateLambda([SelectedDataLayers]() 
+						FExecuteAction::CreateLambda([SelectedDataLayers]()
 						{
 							check(!SelectedDataLayers.IsEmpty());
 							{
-								const FScopedTransaction Transaction(LOCTEXT("AddSelectedActorsToSelectedDataLayers", "Add Selected Actors to Selected Data Layers"));
+								const FScopedTransaction Transaction(LOCTEXT("AddSelectedActorsToSelectedDataLayers", "Add Selected Actor(s) to Selected Data Layer(s)"));
 								UDataLayerEditorSubsystem::Get()->AddSelectedActorsToDataLayers(SelectedDataLayers);
 							}
 						}),
 						FCanExecuteAction::CreateLambda([SelectedDataLayers, bAllSelectedDataLayersCanAddActors, bSelectedActorsAreAllUserManaged]
-						{ 
+						{
 							return !SelectedDataLayers.IsEmpty() && GEditor->GetSelectedActorCount() > 0 && bAllSelectedDataLayersCanAddActors && bSelectedActorsAreAllUserManaged;
 						})
 					));
 
 				if (!Mode->SelectedDataLayerActors.IsEmpty())
 				{
-					Section.AddSubMenu("AddSelectedActorsTo", LOCTEXT("AddSelectedActorsTo", "Add Selected Actors To"), FText(),
+					Section.AddSubMenu("AddSelectedActorsTo", LOCTEXT("AddSelectedActorsTo", "Add Selected Actor(s) To"), FText(),
 						FNewToolMenuDelegate::CreateLambda([Mode](UToolMenu* InSubMenu)
 						{
-							CreateDataLayerPicker(InSubMenu, FOnDataLayerPicked::CreateLambda([Mode](UDataLayerInstance* TargetDataLayer)
-							{
-								check(TargetDataLayer);
-								TArray<AActor*> Actors;
-								Algo::TransformIf(Mode->SelectedDataLayerActors, Actors, [](const FSelectedDataLayerActor& InActor) { return InActor.Value.Get(); }, [](const FSelectedDataLayerActor& InActor) { return const_cast<AActor*>(InActor.Value.Get()); });
-								if (!Actors.IsEmpty())
+							CreateDataLayerPicker(InSubMenu, 
+								FOnDataLayerInstancePicked::CreateLambda([Mode](UDataLayerInstance* TargetDataLayer)
 								{
-									const FScopedTransaction Transaction(LOCTEXT("AddSelectedActorsToDataLayer", "Add Selected Actors to Selected Data Layer"));
-									UDataLayerEditorSubsystem::Get()->AddActorsToDataLayers(Actors, { TargetDataLayer });
-								}
-							}));
-						}));
-				}
-
-				const bool bSelectedDataLayersSupportParenting = Algo::AllOf(SelectedDataLayers, [](const UDataLayerInstance* DataLayerInstance) { return DataLayerInstance->CanHaveParentDataLayer(); });
-				if (!SelectedDataLayers.IsEmpty() && !bSelectedDataLayersContainsLocked && bSelectedDataLayersSupportParenting)
-				{
-					Section.AddSubMenu("MoveSelectedDataLayersTo", LOCTEXT("MoveSelectedDataLayersTo", "Move Data Layers To"), FText(),
-						FNewToolMenuDelegate::CreateLambda([Mode](UToolMenu* InSubMenu)
-						{
-							CreateDataLayerPicker(InSubMenu, FOnDataLayerPicked::CreateLambda([Mode](UDataLayerInstance* TargetDataLayer)
-							{
-								TArray<UDataLayerInstance*> DataLayerInstances;
-								for (TWeakObjectPtr<const UDataLayerInstance>& DataLayerInstance : Mode->SelectedDataLayersSet)
+									check(TargetDataLayer);
+									TArray<AActor*> Actors;
+									Algo::TransformIf(Mode->SelectedDataLayerActors, Actors, [](const FSelectedDataLayerActor& InActor) { return InActor.Value.Get(); }, [](const FSelectedDataLayerActor& InActor) { return const_cast<AActor*>(InActor.Value.Get()); });
+									if (!Actors.IsEmpty())
+									{
+										const FScopedTransaction Transaction(LOCTEXT("AddSelectedActorsToDataLayer", "Add Selected Actors to Selected Data Layer"));
+										UDataLayerEditorSubsystem::Get()->AddActorsToDataLayers(Actors, { TargetDataLayer });
+									}
+								}),
+								FOnShouldFilterDataLayerInstance::CreateLambda([Mode](const UDataLayerInstance* InCandidateDataLayerInstance)
 								{
-									if (DataLayerInstance.IsValid() && !DataLayerInstance->IsLocked() && DataLayerInstance->CanBeChildOf(TargetDataLayer))
+									for (auto& [DataLayerInstance, SelectedActor] : Mode->SelectedDataLayerActors)
 									{
-										DataLayerInstances.Add(const_cast<UDataLayerInstance*>(DataLayerInstance.Get()));
+										if (!InCandidateDataLayerInstance->CanAddActor(const_cast<AActor*>(SelectedActor.Get())))
+										{
+											return false;
+										}
 									}
-								}
-								Mode->SetParentDataLayer(DataLayerInstances, TargetDataLayer);
-							}), /*bShowRoot*/true);
+									return true;
+								})
+							);
 						}));
 				}
 
-				Section.AddSeparator("SectionsSeparator");
-
 				const bool bAllSelectedDataLayersCanRemoveActors = Algo::AllOf(SelectedDataLayers, [](const UDataLayerInstance* DataLayerInstance) { return DataLayerInstance->CanUserRemoveActors(); });
 				FText RemoveActorToolTip = bSelectedActorsAreAllUserManaged ? FText::GetEmpty() : SelectionContainsNonUserManagedActorsText;
-				Section.AddMenuEntry("RemoveSelectedActorsFromSelectedDataLayers", LOCTEXT("RemoveSelectedActorsFromSelectedDataLayersMenu", "Remove Selected Actors from Selected Data Layers"), RemoveActorToolTip, FSlateIcon(),
+				Section.AddMenuEntry("RemoveSelectedActorsFromSelectedDataLayers", LOCTEXT("RemoveSelectedActorsFromSelectedDataLayersMenu", "Remove Selected Actor(s) from Selected Data Layer(s)"), RemoveActorToolTip, FSlateIcon(),
 					FUIAction(
 						FExecuteAction::CreateLambda([SelectedDataLayers]() 
 						{
@@ -1513,35 +1666,21 @@ void FDataLayerMode::RegisterContextMenu()
 						})
 					));
 
-				const bool bAllSelectedDataLayersAreUserManaged = Algo::AllOf(SelectedDataLayers, [](const UDataLayerInstance* DataLayerInstance) { return DataLayerInstance->IsUserManaged(); });
-				Section.AddMenuEntry("DeleteSelectedDataLayers", LOCTEXT("DeleteSelectedDataLayers", "Delete Selected Data Layers"), FText(), FSlateIcon(),
-					FUIAction(
-						FExecuteAction::CreateLambda([Mode, SelectedDataLayers]() 
-						{
-							check(!SelectedDataLayers.IsEmpty());
-							Mode->DeleteDataLayers(SelectedDataLayers);
-						}),
-						FCanExecuteAction::CreateLambda([SelectedDataLayers, bSelectedDataLayersContainsLocked, bAllSelectedDataLayersAreUserManaged]
-						{
-							return !SelectedDataLayers.IsEmpty() && !bSelectedDataLayersContainsLocked && bAllSelectedDataLayersAreUserManaged;
-						})
-					));
-
 				if (UDataLayerEditorSubsystem::Get()->HasDeprecatedDataLayers())
 				{
 					Section.AddMenuEntry("RenameSelectedDataLayer", LOCTEXT("RenameSelectedDataLayer", "Rename Selected Data Layer"), FText(), FSlateIcon(),
 						FUIAction(
 							FExecuteAction::CreateLambda([Mode, SceneOutliner, SelectedDataLayers]() {
-								if (SelectedDataLayers.Num() == 1)
+							if (SelectedDataLayers.Num() == 1)
+							{
+								FSceneOutlinerTreeItemPtr ItemToRename = SceneOutliner->GetTreeItem(SelectedDataLayers[0]);
+								if (ItemToRename.IsValid() && Mode->CanRenameItem(*ItemToRename) && ItemToRename->CanInteract())
 								{
-									FSceneOutlinerTreeItemPtr ItemToRename = SceneOutliner->GetTreeItem(SelectedDataLayers[0]);
-									if (ItemToRename.IsValid() && Mode->CanRenameItem(*ItemToRename) && ItemToRename->CanInteract())
-									{
-										SceneOutliner->SetPendingRenameItem(ItemToRename);
-										SceneOutliner->ScrollItemIntoView(ItemToRename);
-									}
-								}}),
-							FCanExecuteAction::CreateLambda([SelectedDataLayers] { return (SelectedDataLayers.Num() == 1) && !SelectedDataLayers[0]->IsLocked(); })
+									SceneOutliner->SetPendingRenameItem(ItemToRename);
+									SceneOutliner->ScrollItemIntoView(ItemToRename);
+								}
+							}}),
+							FCanExecuteAction::CreateLambda([SelectedDataLayers] { return (SelectedDataLayers.Num() == 1) && !SelectedDataLayers[0]->IsReadOnly(); })
 						));
 				}
 
@@ -1551,50 +1690,38 @@ void FDataLayerMode::RegisterContextMenu()
 			{
 				FToolMenuSection& Section = InMenu->AddSection("DataLayerSelection", LOCTEXT("DataLayerSelection", "Selection"));
 
-				Section.AddMenuEntry("SelectActorsInDataLayers", LOCTEXT("SelectActorsInDataLayers", "Select Actors in Data Layers"), FText(), FSlateIcon(),
+				Section.AddMenuEntry("SelectActorsInDataLayers", LOCTEXT("SelectActorsInDataLayers", "Select Actor(s) in Data Layer(s)"), FText(), FSlateIcon(),
 					FUIAction(
 						FExecuteAction::CreateLambda([SelectedDataLayers]() {
-							check(!SelectedDataLayers.IsEmpty());
-							{
-								const FScopedTransaction Transaction(LOCTEXT("SelectActorsInDataLayers", "Select Actors in Data Layers"));
-								GEditor->SelectNone(/*bNoteSelectionChange*/false, /*bDeselectBSPSurfs*/true);
-								UDataLayerEditorSubsystem::Get()->SelectActorsInDataLayers(SelectedDataLayers, /*bSelect*/true, /*bNotify*/true, /*bSelectEvenIfHidden*/true);
-							}}),
+						check(!SelectedDataLayers.IsEmpty());
+						{
+							const FScopedTransaction Transaction(LOCTEXT("SelectActorsInDataLayers", "Select Actors in Data Layer(s)"));
+							GEditor->SelectNone(/*bNoteSelectionChange*/false, /*bDeselectBSPSurfs*/true);
+							UDataLayerEditorSubsystem::Get()->SelectActorsInDataLayers(SelectedDataLayers, /*bSelect*/true, /*bNotify*/true, /*bSelectEvenIfHidden*/true);
+						}}),
 						FCanExecuteAction::CreateLambda([SelectedDataLayers] { return !SelectedDataLayers.IsEmpty(); })
 					));
 
-				Section.AddMenuEntry("AppendActorsToSelection", LOCTEXT("AppendActorsToSelection", "Append Actors in Data Layer to Selection"), FText(), FSlateIcon(),
+				Section.AddMenuEntry("AddActorsToSelection", LOCTEXT("AddActorsToSelection", "Add Actors in Data Layer(s) to Selection"), FText(), FSlateIcon(),
 					FUIAction(
 						FExecuteAction::CreateLambda([SelectedDataLayers]() {
-							check(!SelectedDataLayers.IsEmpty());
-							{
-								const FScopedTransaction Transaction(LOCTEXT("AppendActorsToSelection", "Append Actors in Data Layer to Selection"));
-								UDataLayerEditorSubsystem::Get()->SelectActorsInDataLayers(SelectedDataLayers, /*bSelect*/true, /*bNotify*/true, /*bSelectEvenIfHidden*/true);
-							}}),
-						FCanExecuteAction::CreateLambda([SelectedDataLayers,bSelectedDataLayersContainsLocked] { return !SelectedDataLayers.IsEmpty() && !bSelectedDataLayersContainsLocked; })
-					));
-
-				Section.AddMenuEntry("DeselectActors", LOCTEXT("DeselectActors", "Deselect Actors in Data Layer"), FText(), FSlateIcon(),
-					FUIAction(
-						FExecuteAction::CreateLambda([SelectedDataLayers]() {
-							check(!SelectedDataLayers.IsEmpty());
-							{
-								const FScopedTransaction Transaction(LOCTEXT("DeselectActors", "Deselect Actors in Data Layer"));
-								UDataLayerEditorSubsystem::Get()->SelectActorsInDataLayers(SelectedDataLayers, /*bSelect*/false, /*bNotifySelectActors*/true);
-							}}),
+						check(!SelectedDataLayers.IsEmpty());
+						{
+							const FScopedTransaction Transaction(LOCTEXT("AppendActorsToSelection", "Append Actors in Data Layer to Selection"));
+							UDataLayerEditorSubsystem::Get()->SelectActorsInDataLayers(SelectedDataLayers, /*bSelect*/true, /*bNotify*/true, /*bSelectEvenIfHidden*/true);
+						}}),
 						FCanExecuteAction::CreateLambda([SelectedDataLayers] { return !SelectedDataLayers.IsEmpty(); })
 					));
 
-				Section.AddMenuEntry("Copy Selected Data Layer Instances(s) File Path", LOCTEXT("CopySelectedDataLayerInstancessFilePath", "Copy Selected Data Layer Instance(s) File Path"), FText(), FSlateIcon(),
+				Section.AddMenuEntry("DeselectActors", LOCTEXT("DeselectActors", "Deselect Actor(s) in Data Layer(s)"), FText(), FSlateIcon(),
 					FUIAction(
 						FExecuteAction::CreateLambda([SelectedDataLayers]() {
-							check(!SelectedDataLayers.IsEmpty());
-							{
-								TArray<const UObject*> Objects;
-								Algo::Transform(SelectedDataLayers, Objects, [](UDataLayerInstance* DataLayerInstance) { return DataLayerInstance; });
-								FExternalPackageHelper::CopyObjectsExternalPackageFilePathToClipboard(Objects);
-							}}),
-						FCanExecuteAction::CreateLambda([bSelectedDataLayersContainsExternalPackage] { return bSelectedDataLayersContainsExternalPackage; })
+						check(!SelectedDataLayers.IsEmpty());
+						{
+							const FScopedTransaction Transaction(LOCTEXT("DeselectActors", "Deselect Actors in Data Layer"));
+							UDataLayerEditorSubsystem::Get()->SelectActorsInDataLayers(SelectedDataLayers, /*bSelect*/false, /*bNotifySelectActors*/true);
+						}}),
+						FCanExecuteAction::CreateLambda([SelectedDataLayers] { return !SelectedDataLayers.IsEmpty(); })
 					));
 			}
 
@@ -1604,22 +1731,22 @@ void FDataLayerMode::RegisterContextMenu()
 				Section.AddMenuEntry("CollapseAllDataLayers", LOCTEXT("CollapseAllDataLayers", "Collapse All Data Layers"), FText(), FSlateIcon(),
 					FUIAction(
 						FExecuteAction::CreateLambda([SceneOutliner,SelectedDataLayers]() {
-							check(!SelectedDataLayers.IsEmpty());
-							{
-								GEditor->SelectNone(/*bNoteSelectionChange*/false, /*bDeselectBSPSurfs*/true);
-								SceneOutliner->CollapseAll();
-							}}),
+						check(!SelectedDataLayers.IsEmpty());
+						{
+							GEditor->SelectNone(/*bNoteSelectionChange*/false, /*bDeselectBSPSurfs*/true);
+							SceneOutliner->CollapseAll();
+						}}),
 						FCanExecuteAction::CreateLambda([SelectedDataLayers] { return !SelectedDataLayers.IsEmpty(); })
 					));
 
 				Section.AddMenuEntry("ExpandAllDataLayers", LOCTEXT("ExpandAllDataLayers", "Expand All Data Layers"), FText(), FSlateIcon(),
 					FUIAction(
 						FExecuteAction::CreateLambda([SceneOutliner,SelectedDataLayers]() {
-							check(!SelectedDataLayers.IsEmpty());
-							{
-								GEditor->SelectNone(/*bNoteSelectionChange*/false, /*bDeselectBSPSurfs*/true);
-								SceneOutliner->ExpandAll();
-							}}),
+						check(!SelectedDataLayers.IsEmpty());
+						{
+							GEditor->SelectNone(/*bNoteSelectionChange*/false, /*bDeselectBSPSurfs*/true);
+							SceneOutliner->ExpandAll();
+						}}),
 						FCanExecuteAction::CreateLambda([SelectedDataLayers] { return !SelectedDataLayers.IsEmpty(); })
 					));
 			}
@@ -1671,7 +1798,7 @@ void FDataLayerMode::RegisterContextMenu()
 						})
 					));
 
-				Section.AddMenuEntry("ClearCurrentDataLayers", LOCTEXT("ClearCurrentDataLayers", "Clear Current Data Layers"), FText(), FSlateIcon(),
+				Section.AddMenuEntry("ClearCurrentDataLayers", LOCTEXT("ClearCurrentDataLayers", "Clear Current Data Layer(s)"), FText(), FSlateIcon(),
 					FUIAction(
 						FExecuteAction::CreateLambda([AllDataLayers]() 
 						{
@@ -1682,7 +1809,8 @@ void FDataLayerMode::RegisterContextMenu()
 								{
 									UDataLayerEditorSubsystem::Get()->RemoveFromActorEditorContext(const_cast<UDataLayerInstance*>(DataLayer));
 								}
-							}}),
+							}
+						}),
 						FCanExecuteAction::CreateLambda([bHasActorEditorContextDataLayers] { return bHasActorEditorContextDataLayers; })
 					));
 			}
@@ -1692,12 +1820,14 @@ void FDataLayerMode::RegisterContextMenu()
 
 				Section.AddMenuEntry("MakeAllDataLayersVisible", LOCTEXT("MakeAllDataLayersVisible", "Make All Data Layers Visible"), FText(), FSlateIcon(),
 					FUIAction(
-						FExecuteAction::CreateLambda([AllDataLayers]() {
+						FExecuteAction::CreateLambda([AllDataLayers]()
+						{
 							check(!AllDataLayers.IsEmpty());
 							{
 								const FScopedTransaction Transaction(LOCTEXT("MakeAllDataLayersVisible", "Make All Data Layers Visible"));
 								UDataLayerEditorSubsystem::Get()->MakeAllDataLayersVisible();
-							}}),
+							}
+						}),
 						FCanExecuteAction::CreateLambda([AllDataLayers] { return !AllDataLayers.IsEmpty(); })
 					));
 			}
@@ -2164,7 +2294,7 @@ FDataLayerPickingMode::FDataLayerPickingMode(const FDataLayerModeParams& Params)
 	SceneOutliner->ExpandAll();
 }
 
-TSharedRef<SWidget> FDataLayerPickingMode::CreateDataLayerPickerWidget(FOnDataLayerPicked OnDataLayerPicked)
+TSharedRef<SWidget> FDataLayerPickingMode::CreateDataLayerPickerWidget(FOnDataLayerInstancePicked OnDataLayerInstancePicked, FOnShouldFilterDataLayerInstance OnShouldFilterDataLayerInstance)
 {
 	// Create mini DataLayers outliner to pick a DataLayer
 	FSceneOutlinerInitializationOptions InitOptions;
@@ -2172,15 +2302,19 @@ TSharedRef<SWidget> FDataLayerPickingMode::CreateDataLayerPickerWidget(FOnDataLa
 	InitOptions.bShowParentTree = true;
 	InitOptions.bShowCreateNewFolder = false;
 	InitOptions.bFocusSearchBoxWhenOpened = true;
+	if (OnShouldFilterDataLayerInstance.IsBound())
+	{
+		InitOptions.Filters->AddFilterPredicate<FDataLayerTreeItem>(FDataLayerTreeItem::FFilterPredicate::CreateLambda([OnShouldFilterDataLayerInstance](const UDataLayerInstance* DataLayerInstance) { return OnShouldFilterDataLayerInstance.Execute(DataLayerInstance); }), FSceneOutlinerFilter::EDefaultBehaviour::Pass);
+	}
 	InitOptions.ColumnMap.Add(FSceneOutlinerBuiltInColumnTypes::Label(), FSceneOutlinerColumnInfo(ESceneOutlinerColumnVisibility::Visible, 2));
-	InitOptions.ModeFactory = FCreateSceneOutlinerMode::CreateLambda([OnDataLayerPicked](SSceneOutliner* Outliner)
+	InitOptions.ModeFactory = FCreateSceneOutlinerMode::CreateLambda([OnDataLayerInstancePicked](SSceneOutliner* Outliner)
 	{ 
-		return new FDataLayerPickingMode(FDataLayerModeParams(Outliner, nullptr, nullptr, FOnSceneOutlinerItemPicked::CreateLambda([OnDataLayerPicked](const FSceneOutlinerTreeItemRef& NewParent)
+		return new FDataLayerPickingMode(FDataLayerModeParams(Outliner, nullptr, nullptr, FOnSceneOutlinerItemPicked::CreateLambda([OnDataLayerInstancePicked](const FSceneOutlinerTreeItemRef& NewParent)
 		{
 			FDataLayerTreeItem* DataLayerItem = NewParent->CastTo<FDataLayerTreeItem>();
 			if (UDataLayerInstance* DataLayer = DataLayerItem ? DataLayerItem->GetDataLayer() : nullptr)
 			{
-				OnDataLayerPicked.ExecuteIfBound(DataLayer);
+				OnDataLayerInstancePicked.ExecuteIfBound(DataLayer);
 			}
 			FSlateApplication::Get().DismissAllMenus();
 		})));
@@ -2212,8 +2346,8 @@ void FDataLayerPickingMode::OnItemSelectionChanged(FSceneOutlinerTreeItemPtr Tre
 			{
 				if (const FDataLayerTreeItem* DataLayerItem = FirstItem->CastTo<FDataLayerTreeItem>())
 				{
-					UDataLayerInstance* DataLayer = DataLayerItem->GetDataLayer();
-					if (DataLayer && !DataLayer->IsLocked())
+					UDataLayerInstance* DataLayerInstance = DataLayerItem->GetDataLayer();
+					if (DataLayerInstance && !DataLayerInstance->IsReadOnly())
 					{
 						OnItemPicked.ExecuteIfBound(FirstItem.ToSharedRef());
 					}
diff --git a/Engine/Source/Editor/DataLayerEditor/Private/DataLayer/DataLayerMode.h b/Engine/Source/Editor/DataLayerEditor/Private/DataLayer/DataLayerMode.h
--- a/Engine/Source/Editor/DataLayerEditor/Private/DataLayer/DataLayerMode.h
+++ b/Engine/Source/Editor/DataLayerEditor/Private/DataLayer/DataLayerMode.h
@@ -56,7 +56,11 @@ struct FDataLayerModeParams
 	FOnSceneOutlinerItemPicked OnItemPicked;
 };
 
-DECLARE_DELEGATE_OneParam(FOnDataLayerPicked, UDataLayerInstance*);
+/** Called when data layer instance is picked in the Data Layer Picker */
+DECLARE_DELEGATE_OneParam(FOnDataLayerInstancePicked, UDataLayerInstance*);
+
+/** Called to check if a data layer instance should be filtered out by Data Layer Picker. Return true to exclude the data layer instance. */
+DECLARE_DELEGATE_RetVal_OneParam(bool, FOnShouldFilterDataLayerInstance, const UDataLayerInstance*);
 
 class FDataLayerMode : public ISceneOutlinerMode
 {
@@ -142,7 +146,7 @@ private:
 	TArray<UDataLayerInstance*> GetSelectedDataLayers(SSceneOutliner* InSceneOutliner) const;
 	void SetParentDataLayer(const TArray<UDataLayerInstance*> DataLayers, UDataLayerInstance* ParentDataLayer) const;
 	void OnLevelSelectionChanged(UObject* Obj);
-	static void CreateDataLayerPicker(UToolMenu* InMenu, FOnDataLayerPicked OnDataLayerPicked, bool bInShowRoot = false);
+	static void CreateDataLayerPicker(UToolMenu* InMenu, FOnDataLayerInstancePicked OnDataLayerInstancePicked, FOnShouldFilterDataLayerInstance OnShouldFilterDataLayerInstance, bool bInShowRoot = false);
 	bool ShouldExpandDataLayer(const UDataLayerInstance* DataLayer) const;
 	bool ContainsSelectedChildDataLayer(const UDataLayerInstance* DataLayer) const;
 	void DeleteDataLayers(const TArray<UDataLayerInstance*>& InDataLayersToDelete);
@@ -150,8 +154,9 @@ private:
 	void CacheSelectedItems(const FSceneOutlinerItemSelection& Selection);
 	UWorld* GetOwningWorld() const;
 	AWorldDataLayers* GetOwningWorldAWorldDataLayers() const;
-	AWorldDataLayers* GetWorldDataLayersFromTreeItem(const ISceneOutlinerTreeItem& TreeItem, bool bDefaultToOwningWorldIfNotFound) const;
-	FSceneOutlinerDragValidationInfo ValidateActorDrop(const ISceneOutlinerTreeItem& DropTarget, bool bMoveOperation = false) const;
+	AWorldDataLayers* GetWorldDataLayersFromTreeItem(const ISceneOutlinerTreeItem& TreeItem) const;
+	UDataLayerInstance* GetDataLayerInstanceFromTreeItem(const ISceneOutlinerTreeItem& TreeItem) const;
+	FSceneOutlinerDragValidationInfo ValidateActorDrop(const ISceneOutlinerTreeItem& DropTarget, TArray<AActor*> PayloadActors, bool bMoveOperation = false) const;
 	FSceneOutlinerDragValidationInfo ValidateDataLayerAssetDrop(const ISceneOutlinerTreeItem& DropTarget, const TArray<const UDataLayerAsset*>& DataLayerAssetsToDrop) const;
 	void OnDataLayerAssetDropped(const TArray<const UDataLayerAsset*>& DroppedDataLayerAsset, ISceneOutlinerTreeItem& DropTarget) const;
 
@@ -210,5 +215,5 @@ public:
 	virtual void OnItemSelectionChanged(FSceneOutlinerTreeItemPtr TreeItem, ESelectInfo::Type SelectionType, const FSceneOutlinerItemSelection& Selection) override;
 	virtual void SynchronizeSelection() override {}
 
-	static TSharedRef<SWidget> CreateDataLayerPickerWidget(FOnDataLayerPicked OnDataLayerPicked);
+	static TSharedRef<SWidget> CreateDataLayerPickerWidget(FOnDataLayerInstancePicked OnDataLayerInstancePicked, FOnShouldFilterDataLayerInstance OnShouldFilterDataLayerInstance = FOnShouldFilterDataLayerInstance());
 };
diff --git a/Engine/Source/Editor/DataLayerEditor/Private/DataLayer/DataLayerNameEditSink.h b/Engine/Source/Editor/DataLayerEditor/Private/DataLayer/DataLayerNameEditSink.h
--- a/Engine/Source/Editor/DataLayerEditor/Private/DataLayer/DataLayerNameEditSink.h
+++ b/Engine/Source/Editor/DataLayerEditor/Private/DataLayer/DataLayerNameEditSink.h
@@ -35,7 +35,7 @@ class FDataLayerNameEditSink : public UE::EditorWidgets::IObjectNameEditSink
 	virtual bool IsObjectDisplayNameReadOnly(UObject* Object) const override
 	{
 		UDataLayerInstance* DataLayerInstance = CastChecked<UDataLayerInstance>(Object);
-		return !DataLayerInstance->CanEditDataLayerShortName() || DataLayerInstance->IsLocked();
+		return !DataLayerInstance->CanEditDataLayerShortName() || DataLayerInstance->IsReadOnly();
 	};
 
 	virtual bool SetObjectDisplayName(UObject* Object, FString DisplayName) override
diff --git a/Engine/Source/Editor/DataLayerEditor/Private/DataLayer/DataLayerOutlinerDeleteButtonColumn.cpp b/Engine/Source/Editor/DataLayerEditor/Private/DataLayer/DataLayerOutlinerDeleteButtonColumn.cpp
--- a/Engine/Source/Editor/DataLayerEditor/Private/DataLayer/DataLayerOutlinerDeleteButtonColumn.cpp
+++ b/Engine/Source/Editor/DataLayerEditor/Private/DataLayer/DataLayerOutlinerDeleteButtonColumn.cpp
@@ -61,14 +61,14 @@ const TSharedRef<SWidget> FDataLayerOutlinerDeleteButtonColumn::ConstructRowWidg
 			.Visibility_Lambda([this, TreeItem, DataLayerActorItem]()
 			{
 				AActor* Actor = DataLayerActorItem->GetActor();
-				const UDataLayerInstance* DataLayer = DataLayerActorItem->GetDataLayer();
-				return (Actor && DataLayer && !DataLayer->IsLocked() && TreeItem->CanInteract()) ? EVisibility::Visible : EVisibility::Collapsed;
+				const UDataLayerInstance* DataLayerInstance = DataLayerActorItem->GetDataLayer();
+				return (Actor && DataLayerInstance && DataLayerInstance->CanRemoveActor(Actor) && DataLayerInstance->CanUserRemoveActors() && TreeItem->CanInteract()) ? EVisibility::Visible : EVisibility::Collapsed;
 			})
 			.OnClicked_Lambda([this, TreeItem, DataLayerActorItem]()
 			{
 				AActor* Actor = DataLayerActorItem->GetActor();
-				const UDataLayerInstance* DataLayer = DataLayerActorItem->GetDataLayer();
-				if (Actor && DataLayer)
+				const UDataLayerInstance* DataLayerInstance = DataLayerActorItem->GetDataLayer();
+				if (Actor && DataLayerInstance)
 				{
 					UDataLayerEditorSubsystem* DataLayerEditorSubsystem = UDataLayerEditorSubsystem::Get();
 					if (auto SceneOutliner = WeakSceneOutliner.IsValid() ? WeakSceneOutliner.Pin() : nullptr)
diff --git a/Engine/Source/Editor/DataLayerEditor/Private/DataLayer/DataLayerOutlinerIsLoadedInEditorColumn.cpp b/Engine/Source/Editor/DataLayerEditor/Private/DataLayer/DataLayerOutlinerIsLoadedInEditorColumn.cpp
--- a/Engine/Source/Editor/DataLayerEditor/Private/DataLayer/DataLayerOutlinerIsLoadedInEditorColumn.cpp
+++ b/Engine/Source/Editor/DataLayerEditor/Private/DataLayer/DataLayerOutlinerIsLoadedInEditorColumn.cpp
@@ -76,7 +76,7 @@ const TSharedRef<SWidget> FDataLayerOutlinerIsLoadedInEditorColumn::ConstructRow
 				{
 					FDataLayerTreeItem* DataLayerTreeItem = TreeItem->CastTo<FDataLayerTreeItem>();
 					const UDataLayerInstance* DataLayerInstance = DataLayerTreeItem->GetDataLayer();
-					const AWorldDataLayers* OuterWorldDataLayers = DataLayerInstance ? DataLayerInstance->GetOuterWorldDataLayers() : nullptr;
+					const AWorldDataLayers* OuterWorldDataLayers = DataLayerInstance ? DataLayerInstance->GetDirectOuterWorldDataLayers() : nullptr;
 					const bool bIsSubWorldDataLayers = OuterWorldDataLayers && OuterWorldDataLayers->IsSubWorldDataLayers();
 					return DataLayerInstance && !DataLayerInstance->IsReadOnly() && !bIsSubWorldDataLayers ? EVisibility::Visible : EVisibility::Collapsed;
 				})
diff --git a/Engine/Source/Editor/DataLayerEditor/Private/DataLayer/DataLayerPropertyTypeCustomization.cpp b/Engine/Source/Editor/DataLayerEditor/Private/DataLayer/DataLayerPropertyTypeCustomization.cpp
--- a/Engine/Source/Editor/DataLayerEditor/Private/DataLayer/DataLayerPropertyTypeCustomization.cpp
+++ b/Engine/Source/Editor/DataLayerEditor/Private/DataLayer/DataLayerPropertyTypeCustomization.cpp
@@ -85,8 +85,8 @@ void FDataLayerPropertyTypeCustomization::CustomizeHeader(TSharedRef<IPropertyHa
 				.IsEnabled_Lambda([this]
 				{
 					FPropertyAccess::Result PropertyAccessResult;
-					const UDataLayerInstance* DataLayer = GetDataLayerFromPropertyHandle(&PropertyAccessResult);
-					return (!DataLayer || !DataLayer->IsLocked());
+					const UDataLayerInstance* DataLayerInstance = GetDataLayerFromPropertyHandle(&PropertyAccessResult);
+					return (!DataLayerInstance || !DataLayerInstance->IsReadOnly());
 				})
 				.ToolTipText(LOCTEXT("ComboButtonTip", "Drag and drop a Data Layer onto this property, or choose one from the drop down."))
 				.OnGetMenuContent(this, &FDataLayerPropertyTypeCustomization::OnGetDataLayerMenu)
@@ -109,8 +109,8 @@ void FDataLayerPropertyTypeCustomization::CustomizeHeader(TSharedRef<IPropertyHa
 				.Visibility_Lambda([this] 
 				{
 					FPropertyAccess::Result PropertyAccessResult;
-					const UDataLayerInstance* DataLayer = GetDataLayerFromPropertyHandle(&PropertyAccessResult);
-					return (DataLayer && DataLayer->IsLocked()) ? EVisibility::Visible : EVisibility::Collapsed;
+					const UDataLayerInstance* DataLayerInstance = GetDataLayerFromPropertyHandle(&PropertyAccessResult);
+					return (DataLayerInstance && DataLayerInstance->IsReadOnly()) ? EVisibility::Visible : EVisibility::Collapsed;
 				})
 				.ColorAndOpacity(this, &FDataLayerPropertyTypeCustomization::GetForegroundColor)
 				.Image(FAppStyle::GetBrush(TEXT("PropertyWindow.Locked")))
@@ -174,8 +174,8 @@ UDataLayerInstance* FDataLayerPropertyTypeCustomization::GetDataLayerFromPropert
 const FSlateBrush* FDataLayerPropertyTypeCustomization::GetDataLayerIcon() const
 {
 	FPropertyAccess::Result PropertyAccessResult;
-	const UDataLayerInstance* DataLayer = GetDataLayerFromPropertyHandle(&PropertyAccessResult);
-	if (!DataLayer)
+	const UDataLayerInstance* DataLayerInstance = GetDataLayerFromPropertyHandle(&PropertyAccessResult);
+	if (!DataLayerInstance)
 	{
 		return FAppStyle::GetBrush(TEXT("DataLayer.Editor"));
 	}
@@ -183,7 +183,7 @@ const FSlateBrush* FDataLayerPropertyTypeCustomization::GetDataLayerIcon() const
 	{
 		return FAppStyle::GetBrush(TEXT("LevelEditor.Tabs.DataLayers"));
 	}
-	return FAppStyle::GetBrush(DataLayer->GetDataLayerIconName());
+	return FAppStyle::GetBrush(DataLayerInstance->GetDataLayerIconName());
 }
 
 FText FDataLayerPropertyTypeCustomization::GetDataLayerText() const
@@ -200,8 +200,8 @@ FText FDataLayerPropertyTypeCustomization::GetDataLayerText() const
 FSlateColor FDataLayerPropertyTypeCustomization::GetForegroundColor() const
 {
 	FPropertyAccess::Result PropertyAccessResult;
-	const UDataLayerInstance* DataLayer = GetDataLayerFromPropertyHandle(&PropertyAccessResult);
-	if (DataLayer && DataLayer->IsLocked())
+	const UDataLayerInstance* DataLayerInstance = GetDataLayerFromPropertyHandle(&PropertyAccessResult);
+	if (DataLayerInstance && DataLayerInstance->IsReadOnly())
 	{
 		return FSceneOutlinerCommonLabelData::DarkColor;
 	}
diff --git a/Engine/Source/Editor/DataLayerEditor/Private/DataLayer/DataLayerPropertyTypeCustomizationHelper.cpp b/Engine/Source/Editor/DataLayerEditor/Private/DataLayer/DataLayerPropertyTypeCustomizationHelper.cpp
--- a/Engine/Source/Editor/DataLayerEditor/Private/DataLayer/DataLayerPropertyTypeCustomizationHelper.cpp
+++ b/Engine/Source/Editor/DataLayerEditor/Private/DataLayer/DataLayerPropertyTypeCustomizationHelper.cpp
@@ -11,7 +11,7 @@ class SWidget;
 
 TSharedRef<SWidget> FDataLayerPropertyTypeCustomizationHelper::CreateDataLayerMenu(TFunction<void(const UDataLayerInstance* DataLayer)> OnDataLayerSelectedFunction)
 {
-	return FDataLayerPickingMode::CreateDataLayerPickerWidget(FOnDataLayerPicked::CreateLambda([OnDataLayerSelectedFunction](UDataLayerInstance* TargetDataLayer)
+	return FDataLayerPickingMode::CreateDataLayerPickerWidget(FOnDataLayerInstancePicked::CreateLambda([OnDataLayerSelectedFunction](UDataLayerInstance* TargetDataLayer)
 	{
 		OnDataLayerSelectedFunction(TargetDataLayer);
 	}));
diff --git a/Engine/Source/Editor/DataLayerEditor/Private/DataLayer/DataLayerTreeItem.cpp b/Engine/Source/Editor/DataLayerEditor/Private/DataLayer/DataLayerTreeItem.cpp
--- a/Engine/Source/Editor/DataLayerEditor/Private/DataLayer/DataLayerTreeItem.cpp
+++ b/Engine/Source/Editor/DataLayerEditor/Private/DataLayer/DataLayerTreeItem.cpp
@@ -25,6 +25,8 @@
 #include "WorldPartition/DataLayer/DataLayerManager.h"
 #include "WorldPartition/DataLayer/DataLayerUtils.h"
 #include "WorldPartition/DataLayer/DataLayerInstanceWithAsset.h"
+#include "WorldPartition/DataLayer/ExternalDataLayerAsset.h"
+#include "WorldPartition/DataLayer/ExternalDataLayerInstance.h"
 #include "ScopedTransaction.h"
 
 template <typename ItemType> class STableRow;
@@ -49,6 +51,7 @@ private:
 	EVisibility GetTypeTextVisibility() const;
 	const FSlateBrush* GetIcon() const;
 	FText GetIconTooltip() const;
+	FSlateColor GetIconColor() const;
 	FSlateColor GetForegroundColor() const;
 	bool OnVerifyItemLabelChanged(const FText& InLabel, FText& OutErrorMessage);
 	void OnLabelCommitted(const FText& InLabel, ETextCommit::Type InCommitInfo);
@@ -88,7 +91,7 @@ bool FDataLayerTreeItem::GetVisibility() const
 bool FDataLayerTreeItem::ShouldShowVisibilityState() const
 {
 	const UDataLayerInstance* DataLayerInstancePtr = DataLayerInstance.Get();
-	const AWorldDataLayers* OuterWorldDataLayers = DataLayerInstancePtr ? DataLayerInstancePtr->GetOuterWorldDataLayers() : nullptr;
+	const AWorldDataLayers* OuterWorldDataLayers = DataLayerInstancePtr ? DataLayerInstancePtr->GetDirectOuterWorldDataLayers() : nullptr;
 	const bool bIsSubWorldDataLayers = OuterWorldDataLayers && OuterWorldDataLayers->IsSubWorldDataLayers();
 	return DataLayerInstancePtr && !DataLayerInstancePtr->IsReadOnly() && !bIsSubWorldDataLayers;
 }
@@ -185,7 +188,7 @@ void SDataLayerTreeLabel::Construct(const FArguments& InArgs, FDataLayerTreeItem
 				SNew(SImage)
 				.Image(this, &SDataLayerTreeLabel::GetIcon)
 				.ToolTipText(this, &SDataLayerTreeLabel::GetIconTooltip)
-				.ColorAndOpacity(FSlateColor::UseForeground())
+				.ColorAndOpacity(this, &SDataLayerTreeLabel::GetIconColor)
 			]
 		]
 		+ SHorizontalBox::Slot()
@@ -201,10 +204,21 @@ void SDataLayerTreeLabel::Construct(const FArguments& InArgs, FDataLayerTreeItem
 		.VAlign(VAlign_Center)
 		[
 			SNew(SImage)
-			.Visibility_Lambda([this] { return (DataLayerPtr.IsValid() && DataLayerPtr->IsLocked() && !DataLayerPtr->IsReadOnly() && DataLayerPtr->GetWorld() && !DataLayerPtr->GetWorld()->IsPlayInEditor()) ? EVisibility::Visible : EVisibility::Collapsed; })
+			.Visibility_Lambda([this] { return (DataLayerPtr.IsValid() && DataLayerPtr->IsReadOnly() && DataLayerPtr->GetWorld() && !DataLayerPtr->GetWorld()->IsPlayInEditor()) ? EVisibility::Visible : EVisibility::Collapsed; })
 			.ColorAndOpacity(FSlateColor::UseForeground())
 			.Image(FAppStyle::GetBrush(TEXT("PropertyWindow.Locked")))
-			.ToolTipText(LOCTEXT("LockedRuntimeDataLayerEditing", "Locked editing. (To allow editing, in Data Layer Outliner, go to Advanced -> Allow Runtime Data Layer Editing)"))
+			.ToolTipText_Lambda([this]
+			{
+				if (const UDataLayerInstance* DataLayerInstance = DataLayerPtr.Get())
+				{
+					FText Reason;
+					if (DataLayerInstance->IsReadOnly(&Reason))
+					{
+						return FText::Format(LOCTEXT("ReadOnlyDataLayerInstance", "{0}"), Reason);
+					}
+				}
+				return FText::GetEmpty();
+			})
 		]
 	];
 }
@@ -322,8 +336,26 @@ const FSlateBrush* SDataLayerTreeLabel::GetIcon() const
 
 FText SDataLayerTreeLabel::GetIconTooltip() const
 {
-	const UDataLayerInstance* DataLayer = DataLayerPtr.Get();
-	return DataLayer ? (DataLayer->IsRuntime() ? FText(LOCTEXT("RuntimeDataLayer", "Runtime Data Layer")) : FText(LOCTEXT("EditorDataLayer", "Editor Data Layer"))) : FText();
+	if (const UDataLayerInstance* DataLayerInstance = DataLayerPtr.Get())
+	{
+		if (DataLayerInstance->IsA<UExternalDataLayerInstance>())
+		{
+			return FText(LOCTEXT("ExternalDataLayer", "External Data Layer"));
+		}
+		if (DataLayerInstance->IsRuntime())
+		{
+			return FText(LOCTEXT("RuntimeDataLayer", "Runtime Data Layer"));
+		}
+		return FText(LOCTEXT("EditorDataLayer", "Editor Data Layer"));
+	}
+	return FText();
+}
+
+FSlateColor SDataLayerTreeLabel::GetIconColor() const
+{
+	const UDataLayerInstance* DataLayerInstance = DataLayerPtr.Get();
+	const UExternalDataLayerInstance* ExternalDataLayerInstance = DataLayerInstance ? DataLayerInstance->GetRootExternalDataLayerInstance() : nullptr;
+	return ExternalDataLayerInstance ? UExternalDataLayerAsset::EditorUXColor : FSlateColor::UseForeground();
 }
 
 FSlateColor SDataLayerTreeLabel::GetForegroundColor() const
@@ -358,7 +390,7 @@ FSlateColor SDataLayerTreeLabel::GetForegroundColor() const
 			return FSceneOutlinerCommonLabelData::DarkColor;
 		}
 	}
-	else if (DataLayerInstance->IsLocked())
+	else if (DataLayerInstance->IsReadOnly())
 	{
 		return FSceneOutlinerCommonLabelData::DarkColor;
 	}
diff --git a/Engine/Source/Editor/DataLayerEditor/Private/DataLayer/SDataLayerOutliner.cpp b/Engine/Source/Editor/DataLayerEditor/Private/DataLayer/SDataLayerOutliner.cpp
--- a/Engine/Source/Editor/DataLayerEditor/Private/DataLayer/SDataLayerOutliner.cpp
+++ b/Engine/Source/Editor/DataLayerEditor/Private/DataLayer/SDataLayerOutliner.cpp
@@ -74,9 +74,9 @@ bool SDataLayerOutliner::CanAddSelectedActorsToSelectedDataLayersClicked() const
 {
 	if (GEditor->GetSelectedActorCount() > 0)
 	{
-		TArray<UDataLayerInstance*> SelectedDataLayers = GetSelectedDataLayers();
-		const bool bSelectedDataLayersContainsLocked = !!SelectedDataLayers.FindByPredicate([](const UDataLayerInstance* DataLayer) { return DataLayer->IsLocked(); });
-		return (!SelectedDataLayers.IsEmpty() && !bSelectedDataLayersContainsLocked);
+		TArray<UDataLayerInstance*> SelectedDataLayerInstances = GetSelectedDataLayers();
+		const bool bSelectedDataLayerInstancesContainsReadOnly = !!SelectedDataLayerInstances.FindByPredicate([](const UDataLayerInstance* DataLayerInstance) { return DataLayerInstance->IsReadOnly(); });
+		return (!SelectedDataLayerInstances.IsEmpty() && !bSelectedDataLayerInstancesContainsReadOnly);
 	}
 	return false;
 }
diff --git a/Engine/Source/Editor/DataLayerEditor/Private/DataLayerEditorModule.cpp b/Engine/Source/Editor/DataLayerEditor/Private/DataLayerEditorModule.cpp
--- a/Engine/Source/Editor/DataLayerEditor/Private/DataLayerEditorModule.cpp
+++ b/Engine/Source/Editor/DataLayerEditor/Private/DataLayerEditorModule.cpp
@@ -7,6 +7,7 @@
 #include "DataLayer/DataLayerNameEditSink.h"
 #include "DataLayer/DataLayerPropertyTypeCustomization.h"
 #include "DataLayer/SDataLayerBrowser.h"
+#include "WorldPartition/DataLayer/ExternalDataLayerAsset.h"
 #include "EditorWidgetsModule.h"
 #include "HAL/Platform.h"
 #include "Modules/ModuleManager.h"
@@ -62,4 +63,9 @@ void FDataLayerEditorModule::SyncDataLayerBrowserToDataLayer(const UDataLayerIns
 bool FDataLayerEditorModule::AddActorToDataLayers(AActor* Actor, const TArray<UDataLayerInstance*>& DataLayers)
 {
 	return UDataLayerEditorSubsystem::Get()->AddActorToDataLayers(Actor, DataLayers);
+}
+
+void FDataLayerEditorModule::SetActorEditorContextCurrentExternalDataLayer(const UExternalDataLayerAsset* InExternalDataLayerAsset)
+{
+	UDataLayerEditorSubsystem::Get()->SetActorEditorContextCurrentExternalDataLayer(InExternalDataLayerAsset);
 }
\ No newline at end of file
diff --git a/Engine/Source/Editor/DataLayerEditor/Public/DataLayer/DataLayerEditorSubsystem.h b/Engine/Source/Editor/DataLayerEditor/Public/DataLayer/DataLayerEditorSubsystem.h
--- a/Engine/Source/Editor/DataLayerEditor/Public/DataLayer/DataLayerEditorSubsystem.h
+++ b/Engine/Source/Editor/DataLayerEditor/Public/DataLayer/DataLayerEditorSubsystem.h
@@ -42,7 +42,10 @@ class ULevel;
 class UObject;
 class UWorld;
 class UWorldPartition;
+class UExternalDataLayerAsset;
+struct FLevelEditorDragDropWorldSurrogateReferencingObject;
 struct FFrame;
+enum class EExternalDataLayerRegistrationState : uint8;
 template<typename TItemType> class IFilter;
 
 USTRUCT(BlueprintType)
@@ -113,6 +116,8 @@ public:
 	//~ End IActorEditorContextClient interface
 	void AddToActorEditorContext(UDataLayerInstance* InDataLayerInstance);
 	void RemoveFromActorEditorContext(UDataLayerInstance* InDataLayerInstance);
+	bool SetActorEditorContextCurrentExternalDataLayer(const UExternalDataLayerAsset* InExternalDataLayerAsset);
+	const UExternalDataLayerAsset* GetActorEditorContextCurrentExternalDataLayer() const;
 
 	TArray<const UDataLayerInstance*> GetDataLayerInstances(const TArray<const UDataLayerAsset*> DataLayerAssets) const;
 
@@ -148,9 +153,19 @@ public:
 	 *
 	 *	@param	Actor	The actor to validate
 	 */
+	UE_DEPRECATED(5.4, "Use IsActorValidForDataLayerInstances instead.")
 	UFUNCTION(BlueprintCallable, Category = DataLayers)
 	bool IsActorValidForDataLayer(AActor* Actor);
 
+	/**
+	 *	Checks to see if the specified actor is in an appropriate state to interact with DataLayers
+	 *
+	 *	@param	Actor				The actor to validate
+	 *  @param  DataLayerInstances	The data layers used to do the validation
+	 */
+	UFUNCTION(BlueprintCallable, Category = DataLayers)
+	bool IsActorValidForDataLayerInstances(AActor* Actor, const TArray<UDataLayerInstance*>& DataLayerInstances);
+
 	/**
 	 * Adds the actor to the DataLayer.
 	 *
@@ -709,6 +724,7 @@ private:
 	/** Delegate handler for FEditorDelegates::PostUndoRedo. It internally calls DataLayerChanged.Broadcast and UpdateAllActorsVisibility to refresh the actors of each DataLayer. */
 	void PostUndoRedo();
 
+	void UpdateRegisteredWorldDelegates();
 	void OnWorldPartitionInitialized(UWorldPartition* InWorldPartition);
 	void OnWorldPartitionUninitialized(UWorldPartition* InWorldPartition);
 	void OnLoadedActorAddedToLevel(AActor& InActor);
@@ -727,6 +743,14 @@ private:
 
 	bool UpdateAllActorsVisibility(const bool bNotifySelectionChange, const bool bRedrawViewports, ULevel* InLevel);
 
+	bool IsActorValidForDataLayerForClasses(AActor* Actor, const TSet<TSubclassOf<UDataLayerInstance>>& DataLayerInstanceClasses);
+
+	// External Data Layer methods
+	void OnActorPreSpawnInitialization(AActor* Actor);
+	static const UExternalDataLayerAsset* GetReferencingWorldSurrogateObjectForObject(UWorld* ReferencingWorld, const FSoftObjectPath& ObjectPath);
+	TUniquePtr<FLevelEditorDragDropWorldSurrogateReferencingObject> OnLevelEditorDragDropWorldSurrogateReferencingObject(UWorld* ReferencingWorld, const FSoftObjectPath& Object);
+	void OnExternalDataLayerAssetRegistrationStateChanged(const UExternalDataLayerAsset* ExternalDataLayerAsset, EExternalDataLayerRegistrationState OldState, EExternalDataLayerRegistrationState NewState);
+
 	/** Contains Data Layers that contain actors that are part of the editor selection */
 	mutable TSet<TWeakObjectPtr<const UDataLayerInstance>> SelectedDataLayersFromEditorSelection;
 
@@ -739,6 +763,9 @@ private:
 	/** When true, next Tick will call UpdateAllActorsVisibility */
 	bool bAsyncUpdateAllActorsVisibility;
 
+	/** When true, next Tick will force invalid editing viewports */
+	bool bAsyncInvalidateViewports;
+
 	/** Fires whenever one or more DataLayer changes */
 	FOnDataLayerChanged DataLayerChanged;
 
@@ -757,16 +784,17 @@ private:
 	FOnActorEditorContextClientChanged ActorEditorContextClientChanged;
 
 	friend class FDataLayersBroadcast;
+	friend struct FExternalDataLayerWorldSurrogateReferencingObject;
 };
 
-template<class DataLayerType, typename ...Args>
-DataLayerType* UDataLayerEditorSubsystem::CreateDataLayerInstance(AWorldDataLayers* WorldDataLayers, Args&&... InArgs)
+template<class DataLayerInstanceType, typename ...Args>
+DataLayerInstanceType* UDataLayerEditorSubsystem::CreateDataLayerInstance(AWorldDataLayers* WorldDataLayers, Args&&... InArgs)
 {
 	UDataLayerInstance* NewDataLayer = nullptr;
 
 	if (WorldDataLayers)
 	{
-		NewDataLayer = WorldDataLayers->CreateDataLayer<DataLayerType>(Forward<Args>(InArgs)...);
+		NewDataLayer = WorldDataLayers->CreateDataLayer<DataLayerInstanceType>(Forward<Args>(InArgs)...);
 	}
 
 	if (NewDataLayer != nullptr)
@@ -774,5 +802,5 @@ DataLayerType* UDataLayerEditorSubsystem::CreateDataLayerInstance(AWorldDataLaye
 		BroadcastDataLayerChanged(EDataLayerAction::Add, NewDataLayer, NAME_None);
 	}
 
-	return CastChecked<DataLayerType>(NewDataLayer);
+	return CastChecked<DataLayerInstanceType>(NewDataLayer);
 }
\ No newline at end of file
diff --git a/Engine/Source/Editor/DataLayerEditor/Public/DataLayerEditorModule.h b/Engine/Source/Editor/DataLayerEditor/Public/DataLayerEditorModule.h
--- a/Engine/Source/Editor/DataLayerEditor/Public/DataLayerEditorModule.h
+++ b/Engine/Source/Editor/DataLayerEditor/Public/DataLayerEditorModule.h
@@ -15,6 +15,7 @@ class FExtender;
 class FUICommandList;
 class SWidget;
 class UDataLayerInstance;
+class UExternalDataLayerAsset;
 
 /**
  * The module holding all of the UI related pieces for DataLayers
@@ -50,6 +51,7 @@ public:
 
 	/* Implement IDataLayerEditorModule */
 	virtual bool AddActorToDataLayers(AActor* Actor, const TArray<UDataLayerInstance*>& DataLayers) override;
+	virtual void SetActorEditorContextCurrentExternalDataLayer(const UExternalDataLayerAsset* InExternalDataLayerAsset) override;
 
 private:
 	TWeakPtr<SWidget> DataLayerBrowser;
diff --git a/Engine/Source/Editor/EditorStyle/Private/StarshipStyle.cpp b/Engine/Source/Editor/EditorStyle/Private/StarshipStyle.cpp
--- a/Engine/Source/Editor/EditorStyle/Private/StarshipStyle.cpp
+++ b/Engine/Source/Editor/EditorStyle/Private/StarshipStyle.cpp
@@ -1067,6 +1067,7 @@ void FStarshipEditorStyle::FStyle::SetupGeneralStyles()
 		Set("ClassIcon.DataLayer", new IMAGE_BRUSH_SVG("Icons/DataLayerEditor", Icon16x16));
 		Set("DataLayer.Editor", new IMAGE_BRUSH_SVG("Icons/DataLayerEditor", Icon16x16));
 		Set("DataLayer.Runtime", new IMAGE_BRUSH_SVG("Icons/DataLayerRuntime", Icon16x16));
+		Set("DataLayer.External", new IMAGE_BRUSH_SVG("Icons/ExternalDataLayer", Icon16x16));
 		Set("DataLayer.LoadedInEditor", new IMAGE_BRUSH_SVG("Icons/DataLayerLoadedInEditor", Icon16x16));
 		Set("DataLayerBrowser.AddSelection", new IMAGE_BRUSH_SVG("Icons/DataLayerAddSelected", Icon16x16));
 		Set("DataLayerBrowser.RemoveSelection", new IMAGE_BRUSH_SVG("Icons/DataLayerRemoveSelected", Icon16x16));
@@ -5909,6 +5910,7 @@ void FStarshipEditorStyle::FStyle::SetupClassIconsAndThumbnails()
 			{TEXT("Emitter")},
 			{TEXT("EmptyActor")},
 			{TEXT("ExponentialHeightFog")},
+			{TEXT("ExternalDataLayerAsset")},
 			{TEXT("FileMediaOutput")},
 			{TEXT("FileMediaSource")},
 			{TEXT("FoliageType_Actor")},
diff --git a/Engine/Source/Editor/SceneOutliner/Private/ActorBrowsingMode.cpp b/Engine/Source/Editor/SceneOutliner/Private/ActorBrowsingMode.cpp
--- a/Engine/Source/Editor/SceneOutliner/Private/ActorBrowsingMode.cpp
+++ b/Engine/Source/Editor/SceneOutliner/Private/ActorBrowsingMode.cpp
@@ -155,7 +155,7 @@ FActorBrowsingMode::FActorBrowsingMode(SSceneOutliner* InSceneOutliner, TWeakObj
 					return true;
 				}
 				
-				for(const UDataLayerInstance* const DataLayerInstance : DataLayerManager->GetDataLayerInstances(ActorDescInstance->GetDataLayerInstanceNames()))
+				for(const UDataLayerInstance* const DataLayerInstance : DataLayerManager->GetDataLayerInstances(ActorDescInstance->GetDataLayerInstanceNames().ToArray()))
 				{
 					if (DataLayerInstance->IsInActorEditorContext())
 					{
diff --git a/Engine/Source/Editor/SceneOutliner/Private/ActorDescTreeItem.cpp b/Engine/Source/Editor/SceneOutliner/Private/ActorDescTreeItem.cpp
--- a/Engine/Source/Editor/SceneOutliner/Private/ActorDescTreeItem.cpp
+++ b/Engine/Source/Editor/SceneOutliner/Private/ActorDescTreeItem.cpp
@@ -6,6 +6,7 @@
 #include "WorldPartition/WorldPartitionActorDesc.h"
 #include "WorldPartition/WorldPartitionActorDescInstance.h"
 #include "WorldPartition/LoaderAdapter/LoaderAdapterPinnedActors.h"
+#include "WorldPartition/DataLayer/ExternalDataLayerAsset.h"
 #include "Widgets/DeclarativeSyntaxSupport.h"
 #include "Layout/WidgetPath.h"
 #include "Framework/Application/MenuStack.h"
@@ -333,4 +334,18 @@ bool FActorDescTreeItem::GetPinnedState() const
 	return false;
 }
 
+UExternalDataLayerAsset* FActorDescTreeItem::GetExternalDataLayerAsset() const
+{
+	if (const FWorldPartitionActorDescInstance* ActorDescInstance = ActorDescHandle.GetInstance())
+	{
+		const FSoftObjectPath& ExternalDataLayerAsset = ActorDescInstance->GetExternalDataLayerAsset();
+		if (CachedExternalDataLayerAsset.ToSoftObjectPath() != ExternalDataLayerAsset)
+		{
+			CachedExternalDataLayerAsset = ExternalDataLayerAsset;
+		}
+		return CachedExternalDataLayerAsset.Get();
+	}
+	return nullptr;
+}
+
 #undef LOCTEXT_NAMESPACE
