Removed invalid checks triggered when removing Data Layers from Actor Editor Context. A case where this could happen is if the actor editor context contained an External Data Layer and the user created a new Level Instance with actors contained in the map.

diff --git a/Engine/Source/Editor/DataLayerEditor/Private/DataLayer/DataLayerEditorSubsystem.cpp b/Engine/Source/Editor/DataLayerEditor/Private/DataLayer/DataLayerEditorSubsystem.cpp
--- a/Engine/Source/Editor/DataLayerEditor/Private/DataLayer/DataLayerEditorSubsystem.cpp
+++ b/Engine/Source/Editor/DataLayerEditor/Private/DataLayer/DataLayerEditorSubsystem.cpp
@@ -472,7 +472,6 @@ void UDataLayerEditorSubsystem::AddToActorEditorContext(UDataLayerInstance* InDa
 
 void UDataLayerEditorSubsystem::RemoveFromActorEditorContext(UDataLayerInstance* InDataLayerInstance)
 {
-	check(InDataLayerInstance->CanBeInActorEditorContext());
 	if (InDataLayerInstance->RemoveFromActorEditorContext())
 	{
 		BroadcastDataLayerChanged(EDataLayerAction::Modify, InDataLayerInstance, NAME_None);
diff --git a/Engine/Source/Runtime/Engine/Private/WorldPartition/DataLayer/WorldDataLayers.cpp b/Engine/Source/Runtime/Engine/Private/WorldPartition/DataLayer/WorldDataLayers.cpp
--- a/Engine/Source/Runtime/Engine/Private/WorldPartition/DataLayer/WorldDataLayers.cpp
+++ b/Engine/Source/Runtime/Engine/Private/WorldPartition/DataLayer/WorldDataLayers.cpp
@@ -894,7 +894,6 @@ bool AWorldDataLayers::RemoveFromActorEditorContext(UDataLayerInstance* InDataLa
 {
 	ON_SCOPE_EXIT { UpdateCurrentColorizedDataLayerInstance(); };
 
-	check(InDataLayerInstance->CanBeInActorEditorContext());
 	check(ContainsDataLayer(InDataLayerInstance));
 
 	const UExternalDataLayerInstance* ExternalDataLayerInstance = Cast<UExternalDataLayerInstance>(InDataLayerInstance);
