Fix bug when creating a Level Instance containing actors assigned to an External Data Layer.

diff --git a/Engine/Source/Editor/DataLayerEditor/Private/DataLayer/DataLayerEditorSubsystem.cpp b/Engine/Source/Editor/DataLayerEditor/Private/DataLayer/DataLayerEditorSubsystem.cpp
--- a/Engine/Source/Editor/DataLayerEditor/Private/DataLayer/DataLayerEditorSubsystem.cpp
+++ b/Engine/Source/Editor/DataLayerEditor/Private/DataLayer/DataLayerEditorSubsystem.cpp
@@ -678,18 +678,22 @@ bool UDataLayerEditorSubsystem::AddActorsToDataLayers(const TArray<AActor*>& Act
 				{
 					// If actor's level WorldDataLayers doesn't match this DataLayerInstance outer WorldDataLayers, 
 					// Make sure that a DataLayer Instance for this Data Layer Asset exists in the Actor's level WorldDataLayers.
-					AWorldDataLayers* TargetWorldDataLayers = Actor->GetLevel()->GetWorldDataLayers();
-					if (TargetWorldDataLayers != DataLayerInstance->GetOuterWorldDataLayers())
+					// Skip this for External Data Layers as they are only applied to the parent LevelInstance actor
+					if (!DataLayerInstanceWithAsset->IsA<UExternalDataLayerInstance>())
 					{
-						UDataLayerManager* DataLayerManager = UDataLayerManager::GetDataLayerManager(Actor);
-						if (ensureMsgf(DataLayerManager, TEXT("No DataLayerManager found for Actor %s, can't add actors to data layers."), *Actor->GetName()))
+						AWorldDataLayers* TargetWorldDataLayers = Actor->GetLevel()->GetWorldDataLayers();
+						if (TargetWorldDataLayers != DataLayerInstance->GetOuterWorldDataLayers())
 						{
-							DataLayerInstance = DataLayerManager->GetDataLayerInstance(DataLayerInstanceWithAsset->GetAsset());
-
-							bool bDataLayerInstanceExistsInActorLevel = DataLayerInstance != nullptr;
-							if (!bDataLayerInstanceExistsInActorLevel)
+							UDataLayerManager* DataLayerManager = UDataLayerManager::GetDataLayerManager(Actor);
+							if (ensureMsgf(DataLayerManager, TEXT("No DataLayerManager found for Actor %s, can't add actors to data layers."), *Actor->GetName()))
 							{
-								DataLayerInstance = CreateDataLayerInstance<UDataLayerInstanceWithAsset>(TargetWorldDataLayers, DataLayerInstanceWithAsset->GetAsset());
+								DataLayerInstance = DataLayerManager->GetDataLayerInstance(DataLayerInstanceWithAsset->GetAsset());
+
+								bool bDataLayerInstanceExistsInActorLevel = DataLayerInstance != nullptr;
+								if (!bDataLayerInstanceExistsInActorLevel)
+								{
+									DataLayerInstance = CreateDataLayerInstance<UDataLayerInstanceWithAsset>(TargetWorldDataLayers, DataLayerInstanceWithAsset->GetAsset());
+								}
 							}
 						}
 					}
