Fixed other Localization Text conflicts

diff --git a/Engine/Source/Editor/DataLayerEditor/Private/DataLayer/DataLayerMode.cpp b/Engine/Source/Editor/DataLayerEditor/Private/DataLayer/DataLayerMode.cpp
--- a/Engine/Source/Editor/DataLayerEditor/Private/DataLayer/DataLayerMode.cpp
+++ b/Engine/Source/Editor/DataLayerEditor/Private/DataLayer/DataLayerMode.cpp
@@ -510,7 +510,7 @@ void FDataLayerMode::DeleteDataLayers(const TArray<UDataLayerInstance*>& InDataL
 	}
 
 	{
-		const FScopedTransaction Transaction(LOCTEXT("DeleteSelectedDataLayers", "Delete Selected Data Layers"));
+		const FScopedTransaction Transaction(LOCTEXT("DeleteSelectedDataLayers", "Delete Selected Data Layer(s)"));
 		DataLayerEditorSubsystem->DeleteDataLayers(DataLayersToDelete);
 	}
 
@@ -809,7 +809,7 @@ FSceneOutlinerDragValidationInfo FDataLayerMode::ValidateDataLayerAssetDrop(cons
 	UExternalDataLayerManager* ExternalDataLayerManager = UExternalDataLayerManager::GetExternalDataLayerManager(OuterWorld);
 	if (!ExternalDataLayerManager)
 	{
-		return FSceneOutlinerDragValidationInfo(ESceneOutlinerDropCompatibility::Incompatible, LOCTEXT("CantCreateDataLayerInstanceWithExternalDataLayerAsset", "Cannot create External Data Layer Instance : Partitioned world doesn't support External Data Layers."));
+		return FSceneOutlinerDragValidationInfo(ESceneOutlinerDropCompatibility::Incompatible, LOCTEXT("CantCreateDataLayerInstanceWithExternalDataLayerAssetNotSupported", "Cannot create External Data Layer Instance : Partitioned world doesn't support External Data Layers."));
 	}
 	if (!GetDefault<UEditorExperimentalSettings>()->bEnableWorldPartitionExternalDataLayers)
 	{
@@ -1600,7 +1600,7 @@ void FDataLayerMode::RegisterContextMenu()
 
 				const bool bAllSelectedDataLayersCanAddActors = Algo::AllOf(SelectedDataLayers, [](const UDataLayerInstance* DataLayerInstance) { return DataLayerInstance->CanUserAddActors(); });
 				FText AddActorToolTip = bSelectedActorsAreAllUserManaged ? FText::GetEmpty() : SelectionContainsNonUserManagedActorsText;
-				Section.AddMenuEntry("AddSelectedActorsToSelectedDataLayers", LOCTEXT("AddSelectedActorsToSelectedDataLayersMenu", "Add Selected Actor(s) to Selected Data Layer(s)"), AddActorToolTip, FSlateIcon(),
+				Section.AddMenuEntry("AddSelectedActorsToSelectedDataLayers", LOCTEXT("AddSelectedActorsToSelectedDataLayers", "Add Selected Actor(s) to Selected Data Layer(s)"), AddActorToolTip, FSlateIcon(),
 					FUIAction(
 						FExecuteAction::CreateLambda([SelectedDataLayers]()
 						{
@@ -1695,7 +1695,7 @@ void FDataLayerMode::RegisterContextMenu()
 						FExecuteAction::CreateLambda([SelectedDataLayers]() {
 						check(!SelectedDataLayers.IsEmpty());
 						{
-							const FScopedTransaction Transaction(LOCTEXT("SelectActorsInDataLayers", "Select Actors in Data Layer(s)"));
+							const FScopedTransaction Transaction(LOCTEXT("SelectActorsInDataLayers", "Select Actor(s) in Data Layer(s)"));
 							GEditor->SelectNone(/*bNoteSelectionChange*/false, /*bDeselectBSPSurfs*/true);
 							UDataLayerEditorSubsystem::Get()->SelectActorsInDataLayers(SelectedDataLayers, /*bSelect*/true, /*bNotify*/true, /*bSelectEvenIfHidden*/true);
 						}}),
@@ -1718,7 +1718,7 @@ void FDataLayerMode::RegisterContextMenu()
 						FExecuteAction::CreateLambda([SelectedDataLayers]() {
 						check(!SelectedDataLayers.IsEmpty());
 						{
-							const FScopedTransaction Transaction(LOCTEXT("DeselectActors", "Deselect Actors in Data Layer"));
+							const FScopedTransaction Transaction(LOCTEXT("DeselectActors", "Deselect Actor(s) in Data Layer(s)"));
 							UDataLayerEditorSubsystem::Get()->SelectActorsInDataLayers(SelectedDataLayers, /*bSelect*/false, /*bNotifySelectActors*/true);
 						}}),
 						FCanExecuteAction::CreateLambda([SelectedDataLayers] { return !SelectedDataLayers.IsEmpty(); })
@@ -1804,7 +1804,7 @@ void FDataLayerMode::RegisterContextMenu()
 						{
 							check(!AllDataLayers.IsEmpty());
 							{
-								const FScopedTransaction Transaction(LOCTEXT("ClearCurrentDataLayers", "Clear Current Data Layers"));
+								const FScopedTransaction Transaction(LOCTEXT("ClearCurrentDataLayers", "Clear Current Data Layer(s)"));
 								for (const UDataLayerInstance* DataLayer : AllDataLayers)
 								{
 									UDataLayerEditorSubsystem::Get()->RemoveFromActorEditorContext(const_cast<UDataLayerInstance*>(DataLayer));
diff --git a/Engine/Source/Editor/DataLayerEditor/Private/DataLayer/SDataLayerOutliner.cpp b/Engine/Source/Editor/DataLayerEditor/Private/DataLayer/SDataLayerOutliner.cpp
--- a/Engine/Source/Editor/DataLayerEditor/Private/DataLayer/SDataLayerOutliner.cpp
+++ b/Engine/Source/Editor/DataLayerEditor/Private/DataLayer/SDataLayerOutliner.cpp
@@ -91,7 +91,7 @@ FReply SDataLayerOutliner::OnAddSelectedActorsToSelectedDataLayersClicked()
 	if (CanAddSelectedActorsToSelectedDataLayersClicked())
 	{
 		TArray<UDataLayerInstance*> SelectedDataLayers = GetSelectedDataLayers();
-		const FScopedTransaction Transaction(LOCTEXT("AddSelectedActorsToSelectedDataLayers", "Add Selected Actors to Selected Data Layers"));
+		const FScopedTransaction Transaction(LOCTEXT("AddSelectedActorsToSelectedDataLayers", "Add Selected Actor(s) to Selected Data Layer(s)"));
 		UDataLayerEditorSubsystem::Get()->AddSelectedActorsToDataLayers(SelectedDataLayers);
 	}
 	return FReply::Handled();
