Fix ensure when editing CDO properties in FDataLayersBroadcast::OnObjectPostEditChange.
It can be called on a CDO and GetWorld() can be overrided using the "Within=World" macro (which ensures when being called on a CDO).

diff --git a/Engine/Source/Editor/DataLayerEditor/Private/DataLayer/DataLayerEditorSubsystem.cpp b/Engine/Source/Editor/DataLayerEditor/Private/DataLayer/DataLayerEditorSubsystem.cpp
--- a/Engine/Source/Editor/DataLayerEditor/Private/DataLayer/DataLayerEditorSubsystem.cpp
+++ b/Engine/Source/Editor/DataLayerEditor/Private/DataLayer/DataLayerEditorSubsystem.cpp
@@ -191,18 +191,22 @@ void FDataLayersBroadcast::Initialize()
 
 void FDataLayersBroadcast::OnObjectPostEditChange(UObject* Object, FPropertyChangedEvent& PropertyChangedEvent)
 {
+	if (!Object || Object->IsTemplate() || PropertyChangedEvent.ChangeType == EPropertyChangeType::Interactive)
+	{
+		return;
+	}
+	
 	// Ignore changed on game world objects
-	UWorld* World = Object ? Object->GetWorld() : nullptr;
+	const UWorld* World = Object->GetWorld();
 	const bool bIsGameWorld = World && World->IsGameWorld();
-
-	if (Object && !bIsGameWorld && (PropertyChangedEvent.ChangeType != EPropertyChangeType::Interactive))
+	if (!bIsGameWorld)
 	{
 		bool bRefresh = false;
 		if (Object->IsA<UDataLayerInstance>() || Object->IsA<UDataLayerAsset>())
 		{
 			bRefresh = true;
 		}
-		else if (AActor* Actor = Cast<AActor>(Object))
+		else if (const AActor* Actor = Cast<AActor>(Object))
 		{
 			bRefresh = Actor->IsPropertyChangedAffectingDataLayers(PropertyChangedEvent);
 		}
