UDataLayerEditorSubsystem now exposes to Blueprint functions to Add/Remove Data Layers to the Actor Editor Context.

diff --git a/Engine/Source/Editor/DataLayerEditor/Public/DataLayer/DataLayerEditorSubsystem.h b/Engine/Source/Editor/DataLayerEditor/Public/DataLayer/DataLayerEditorSubsystem.h
--- a/Engine/Source/Editor/DataLayerEditor/Public/DataLayer/DataLayerEditorSubsystem.h
+++ b/Engine/Source/Editor/DataLayerEditor/Public/DataLayer/DataLayerEditorSubsystem.h
@@ -114,9 +114,13 @@ public:
 	virtual TSharedRef<SWidget> GetActorEditorContextWidget(UWorld* InWorld) const override;
 	virtual FOnActorEditorContextClientChanged& GetOnActorEditorContextClientChanged() override { return ActorEditorContextClientChanged; }
 	//~ End IActorEditorContextClient interface
+	UFUNCTION(BlueprintCallable, Category = DataLayers)
 	void AddToActorEditorContext(UDataLayerInstance* InDataLayerInstance);
+	UFUNCTION(BlueprintCallable, Category = DataLayers)
 	void RemoveFromActorEditorContext(UDataLayerInstance* InDataLayerInstance);
+	UFUNCTION(BlueprintCallable, Category = DataLayers)
 	bool SetActorEditorContextCurrentExternalDataLayer(const UExternalDataLayerAsset* InExternalDataLayerAsset);
+	UFUNCTION(BlueprintCallable, Category = DataLayers)
 	const UExternalDataLayerAsset* GetActorEditorContextCurrentExternalDataLayer() const;
 
 	TArray<const UDataLayerInstance*> GetDataLayerInstances(const TArray<const UDataLayerAsset*> DataLayerAssets) const;
