Data Layer Outliner QoL
- Double click on unloaded actor selects and focus on the item in the outliner
- Added Pin/Unpin option on a Data Layer to allow pinning/unpinning all child actors

diff --git a/Engine/Source/Editor/DataLayerEditor/Private/DataLayer/DataLayerEditorSubsystem.cpp b/Engine/Source/Editor/DataLayerEditor/Private/DataLayer/DataLayerEditorSubsystem.cpp
--- a/Engine/Source/Editor/DataLayerEditor/Private/DataLayer/DataLayerEditorSubsystem.cpp
+++ b/Engine/Source/Editor/DataLayerEditor/Private/DataLayer/DataLayerEditorSubsystem.cpp
@@ -926,6 +926,47 @@ bool UDataLayerEditorSubsystem::SelectActorsInDataLayers(const TArray<UDataLayer
 	return bChangesOccurred;
 }
 
+void UDataLayerEditorSubsystem::SetActorsPinStateInDataLayers(const TArray<UDataLayerInstance*>& DataLayerInstances, const bool bPinned)
+{
+	UWorld* World = GetWorld();
+	if (UWorldPartition* WorldPartition = World ? World->GetWorldPartition() : nullptr)
+	{
+		TSet<FName> DataLayerInstanceNames;
+		Algo::TransformIf(DataLayerInstances, DataLayerInstanceNames, [](UDataLayerInstance* DataLayerInstance) { return !!DataLayerInstance; }, [](UDataLayerInstance* DataLayerInstance) { return DataLayerInstance->GetDataLayerFName(); });
+
+		if (DataLayerInstanceNames.Num())
+		{
+			TArray<FGuid> ActorGuids;
+			for (FActorDescContainerInstanceCollection::TIterator<> Iterator(WorldPartition); Iterator; ++Iterator)
+			{
+				const FDataLayerInstanceNames& ActorDescDataLayerInstanceNames = Iterator->GetDataLayerInstanceNames();
+				if (ActorDescDataLayerInstanceNames.Num())
+				{
+					for (const FName& DataLayerInstance : DataLayerInstanceNames)
+					{
+						if (ActorDescDataLayerInstanceNames.Contains(DataLayerInstance))
+						{
+							ActorGuids.Add(Iterator->GetGuid());
+						}
+					}
+				}
+			}
+
+			if (ActorGuids.Num())
+			{
+				if (bPinned)
+				{
+					WorldPartition->PinActors(ActorGuids);
+				}
+				else
+				{
+					WorldPartition->UnpinActors(ActorGuids);
+				}
+			}
+		}
+	}
+}
+
 ///////////////////////////////////////////////////////////////////////////////////////////////////////////////////////
 //
 // Operations on actor viewport visibility regarding DataLayers
diff --git a/Engine/Source/Editor/DataLayerEditor/Private/DataLayer/DataLayerMode.cpp b/Engine/Source/Editor/DataLayerEditor/Private/DataLayer/DataLayerMode.cpp
--- a/Engine/Source/Editor/DataLayerEditor/Private/DataLayer/DataLayerMode.cpp
+++ b/Engine/Source/Editor/DataLayerEditor/Private/DataLayer/DataLayerMode.cpp
@@ -449,6 +449,10 @@ void FDataLayerMode::OnItemDoubleClick(FSceneOutlinerTreeItemPtr Item)
 			GEditor->MoveViewportCamerasToActor(*Actor, /*bActiveViewportOnly*/false);
 		}
 	}
+	else if (FDataLayerActorDescTreeItem* DataLayerActorDescItem = Item->CastTo<FDataLayerActorDescTreeItem>())
+	{
+		GEditor->BroadcastSelectUnloadedActors({ DataLayerActorDescItem->GetGuid() });
+	}
 }
 
 void FDataLayerMode::DeleteItems(const TArray<TWeakPtr<ISceneOutlinerTreeItem>>& Items)
@@ -1719,6 +1723,32 @@ void FDataLayerMode::RegisterContextMenu()
 					));
 			}
 
+			{
+				FToolMenuSection& Section = InMenu->AddSection("DataLayerLoading", LOCTEXT("DataLayerLoading", "Loading"));
+
+				Section.AddMenuEntry("PinActorsInDataLayers", LOCTEXT("PinActorsInDataLayers", "Pin Actor(s) in Data Layer(s)"), FText(), FSlateIcon(),
+					FUIAction(
+						FExecuteAction::CreateLambda([SelectedDataLayers]() 
+						{
+							check(!SelectedDataLayers.IsEmpty());
+							const FScopedTransaction Transaction(LOCTEXT("PinActorsInDataLayers", "Pin Actor(s) in Data Layer(s)"));
+							UDataLayerEditorSubsystem::Get()->SetActorsPinStateInDataLayers(SelectedDataLayers, true);
+						}),
+						FCanExecuteAction::CreateLambda([SelectedDataLayers] { return !SelectedDataLayers.IsEmpty(); })
+					));
+
+				Section.AddMenuEntry("UnpinActorsInDataLayers", LOCTEXT("UnpinActorsInDataLayers", "Unpin Actor(s) in Data Layer(s)"), FText(), FSlateIcon(),
+					FUIAction(
+						FExecuteAction::CreateLambda([SelectedDataLayers]()
+							{
+								check(!SelectedDataLayers.IsEmpty());
+								const FScopedTransaction Transaction(LOCTEXT("UnpinActorsInDataLayers", "Unpin Actor(s) in Data Layer(s)"));
+								UDataLayerEditorSubsystem::Get()->SetActorsPinStateInDataLayers(SelectedDataLayers, false);
+							}),
+						FCanExecuteAction::CreateLambda([SelectedDataLayers] { return !SelectedDataLayers.IsEmpty(); })
+					));
+			}
+
 			{
 				FToolMenuSection& Section = InMenu->AddSection("DataLayerExpansion", LOCTEXT("DataLayerExpansion", "Expansion"));
 
diff --git a/Engine/Source/Editor/DataLayerEditor/Public/DataLayer/DataLayerEditorSubsystem.h b/Engine/Source/Editor/DataLayerEditor/Public/DataLayer/DataLayerEditorSubsystem.h
--- a/Engine/Source/Editor/DataLayerEditor/Public/DataLayer/DataLayerEditorSubsystem.h
+++ b/Engine/Source/Editor/DataLayerEditor/Public/DataLayer/DataLayerEditorSubsystem.h
@@ -362,6 +362,15 @@ public:
 	 */
 	bool SelectActorsInDataLayers(const TArray<UDataLayerInstance*>& DataLayers, const bool bSelect, const bool bNotify, const bool bSelectEvenIfHidden, const TSharedPtr<FActorFilter>& Filter);
 
+	/** 
+	 * Pin/unpin actors belonging to the DataLayers.
+	 * 
+	 * @param	DataLayerInstances				A valid list of Data Layer Instances.
+	 * @param	bPinned							If true actors are pinned; if false, actors are unpinned.
+	 * 
+	 */
+	void SetActorsPinStateInDataLayers(const TArray<UDataLayerInstance*>& DataLayerInstances, const bool bPinned);
+
 
 	/////////////////////////////////////////////////
 	// Operations on actor viewport visibility regarding DataLayers
