Added menu option "Browse to Plugin Asset" for Data Layer instances referencing an External Data Layer asset.

diff --git a/Engine/Source/Editor/DataLayerEditor/Private/DataLayer/DataLayerMode.cpp b/Engine/Source/Editor/DataLayerEditor/Private/DataLayer/DataLayerMode.cpp
--- a/Engine/Source/Editor/DataLayerEditor/Private/DataLayer/DataLayerMode.cpp
+++ b/Engine/Source/Editor/DataLayerEditor/Private/DataLayer/DataLayerMode.cpp
@@ -1860,6 +1860,29 @@ void FDataLayerMode::RegisterContextMenu()
 			{
 				FToolMenuSection& Section = InMenu->AddSection("AssetOptionsSection", LOCTEXT("AssetOptionsText", "Asset Options"));
 				Section.AddMenuEntryWithCommandList(FGlobalEditorCommonCommands::Get().FindInContentBrowser, Mode->Commands);
+				const UExternalDataLayerInstance* ExternalDataLayerInstance = SelectedDataLayers.Num() == 1 ? Cast<UExternalDataLayerInstance>(SelectedDataLayers[0]) : nullptr;
+				if (const UExternalDataLayerAsset* ExternalDataLayerAsset = ExternalDataLayerInstance ? ExternalDataLayerInstance->GetExternalDataLayerAsset() : nullptr)
+				{
+					Section.AddMenuEntry("BrowseToPluginAsset", LOCTEXT("BrowseToPluginAsset", "Browse to Plugin Asset"), FText(), FSlateIcon(FAppStyle::GetAppStyleSetName(), "SystemWideCommands.FindInContentBrowser"),
+						FUIAction(
+							FExecuteAction::CreateLambda([ExternalDataLayerAsset]()
+							{
+								UExternalDataLayerEngineSubsystem& ExternalDataLayerEngineSubsystem = UExternalDataLayerEngineSubsystem::Get();
+								TArray<UObject*> Objects;;
+								for (UObject* Client : ExternalDataLayerEngineSubsystem.GetClientsForExternalDataLayerAsset(ExternalDataLayerAsset))
+								{
+									if (UDataAsset* DataAsset = Client->GetTypedOuter<UDataAsset>())
+									{
+										Objects.Add(DataAsset);
+									}
+								}
+								if (!Objects.IsEmpty())
+								{
+									GEditor->SyncBrowserToObjects(Objects);
+								}
+							})
+						));
+				}
 			}
 		}));
 	}
diff --git a/Engine/Source/Runtime/Engine/Private/WorldPartition/DataLayer/ExternalDataLayerEngineSubsystem.cpp b/Engine/Source/Runtime/Engine/Private/WorldPartition/DataLayer/ExternalDataLayerEngineSubsystem.cpp
--- a/Engine/Source/Runtime/Engine/Private/WorldPartition/DataLayer/ExternalDataLayerEngineSubsystem.cpp
+++ b/Engine/Source/Runtime/Engine/Private/WorldPartition/DataLayer/ExternalDataLayerEngineSubsystem.cpp
@@ -268,6 +268,22 @@ void UExternalDataLayerEngineSubsystem::UnregisterExternalDataLayerAsset(const U
 	}
 };
 
+TArray<UObject*> UExternalDataLayerEngineSubsystem::GetClientsForExternalDataLayerAsset(const UExternalDataLayerAsset* InExternalDataLayerAsset) const
+{
+	const FRegisteredExternalDataLayers* RegisteredEDL = ExternalDataLayerAssets.Find(InExternalDataLayerAsset);
+	TSet<UObject*> Clients;
+	for (const FObjectKey Client : RegisteredEDL->RegisteredClients)
+	{
+		Clients.Add(Client.ResolveObjectPtr());
+	}
+	for (const FObjectKey Client : RegisteredEDL->ActiveClients)
+	{
+		Clients.Add(Client.ResolveObjectPtr());
+	}
+	Clients.Remove(nullptr);
+	return Clients.Array();
+}
+
 EExternalDataLayerRegistrationState UExternalDataLayerEngineSubsystem::GetExternalDataLayerAssetRegistrationState(const UExternalDataLayerAsset* InExternalDataLayerAsset) const
 {
 	const FRegisteredExternalDataLayers* RegisteredEDL = ExternalDataLayerAssets.Find(InExternalDataLayerAsset);
diff --git a/Engine/Source/Runtime/Engine/Public/WorldPartition/DataLayer/ExternalDataLayerEngineSubsystem.h b/Engine/Source/Runtime/Engine/Public/WorldPartition/DataLayer/ExternalDataLayerEngineSubsystem.h
--- a/Engine/Source/Runtime/Engine/Public/WorldPartition/DataLayer/ExternalDataLayerEngineSubsystem.h
+++ b/Engine/Source/Runtime/Engine/Public/WorldPartition/DataLayer/ExternalDataLayerEngineSubsystem.h
@@ -48,6 +48,7 @@ public:
 	ENGINE_API bool IsExternalDataLayerAssetRegistered(const UExternalDataLayerAsset* InExternalDataLayerAsset, const UObject* InClient = nullptr) const;
 	ENGINE_API bool IsExternalDataLayerAssetActive(const UExternalDataLayerAsset* InExternalDataLayerAsset, const UObject* InClient = nullptr) const;
 	ENGINE_API EExternalDataLayerRegistrationState GetExternalDataLayerAssetRegistrationState(const UExternalDataLayerAsset* InExternalDataLayerAsset) const;
+	ENGINE_API TArray<UObject*> GetClientsForExternalDataLayerAsset(const UExternalDataLayerAsset* InExternalDataLayerAsset) const;
 
 	DECLARE_MULTICAST_DELEGATE_ThreeParams(FExternalDataLayerAssetRegistrationStateChangedEventDelegate, const UExternalDataLayerAsset* ExternalDataLayerAsset, EExternalDataLayerRegistrationState OldState, EExternalDataLayerRegistrationState NewState);
 	FExternalDataLayerAssetRegistrationStateChangedEventDelegate OnExternalDataLayerAssetRegistrationStateChanged;
