Create New Data Layer With Asset now filters and show only allowed Data Layer Assets

diff --git a/Engine/Source/Editor/DataLayerEditor/Private/DataLayer/DataLayerMode.cpp b/Engine/Source/Editor/DataLayerEditor/Private/DataLayer/DataLayerMode.cpp
--- a/Engine/Source/Editor/DataLayerEditor/Private/DataLayer/DataLayerMode.cpp
+++ b/Engine/Source/Editor/DataLayerEditor/Private/DataLayer/DataLayerMode.cpp
@@ -780,6 +780,19 @@ AWorldDataLayers* FDataLayerMode::GetWorldDataLayersFromTreeItem(const ISceneOut
 	return WorldDataLayers;
 }
 
+bool FDataLayerMode::CanReferenceDataLayerAssets(const AWorldDataLayers* InWorldDataLayers, const TArray<const UDataLayerAsset*>& InReferencedDataLayerAssets, FText* OutFailureReason) const
+{
+	const AWorldDataLayers* ReferencingWorldDataLayers = InWorldDataLayers ? InWorldDataLayers : GetOwningWorldAWorldDataLayers();
+	for (const UDataLayerAsset* DataLayerAsset : InReferencedDataLayerAssets)
+	{
+		if (!ReferencingWorldDataLayers->CanReferenceDataLayerAsset(DataLayerAsset, OutFailureReason))
+		{
+			return false;
+		}
+	}
+	return true;
+}
+
 FSceneOutlinerDragValidationInfo FDataLayerMode::ValidateDataLayerAssetDrop(const ISceneOutlinerTreeItem& DropTarget, const TArray<const UDataLayerAsset*>& DataLayerAssetsToDrop) const
 {
 	check(!DataLayerAssetsToDrop.IsEmpty());
@@ -834,31 +847,21 @@ FSceneOutlinerDragValidationInfo FDataLayerMode::ValidateDataLayerAssetDrop(cons
 		}
 	}
 
-	auto PassesAssetReferenceFiltering = [](const UObject* InReferencingObject, const UDataLayerAsset* InDataLayerAsset, FText* OutReason)
+	// Special case where we receive external data layer asset(s)
+	if (Algo::AllOf(DataLayerAssetsToDrop, [](const UDataLayerAsset* DataLayerAsset) { return DataLayerAsset && DataLayerAsset->IsA<UExternalDataLayerAsset>(); }))
 	{
-		if (InReferencingObject->IsA<AWorldDataLayers>() && InDataLayerAsset->IsA<UExternalDataLayerAsset>())
+		// Only allow to create Data Layer Instance referencing an External Data Layer Asset if the target WorldDataLayers is the main one
+		if (DropTargetWorldDataLayers == GetOwningWorldAWorldDataLayers())
 		{
-			return true;
+			return FSceneOutlinerDragValidationInfo(ESceneOutlinerDropCompatibility::Compatible, LOCTEXT("CreateAllExternalDataLayerFromAssetDrop", "Create External Data Layer Instances"));
 		}
-		FAssetReferenceFilterContext AssetReferenceFilterContext;
-		AssetReferenceFilterContext.ReferencingAssets.Add(FAssetData(InReferencingObject));
-		TSharedPtr<IAssetReferenceFilter> AssetReferenceFilter = GEditor->MakeAssetReferenceFilter(AssetReferenceFilterContext);
-		return AssetReferenceFilter.IsValid() ? AssetReferenceFilter->PassesFilter(FAssetData(InDataLayerAsset), OutReason) : true;
-	};
+	}
 
 	// Check if can reference Data Layer Asset
-	const UExternalDataLayerInstance* RootExternalDataLayerInstance = DropTargetDataLayerWithAsset ? DropTargetDataLayerWithAsset->GetRootExternalDataLayerInstance() : nullptr;
-	const UExternalDataLayerAsset* RootExternalDataLayerAsset = RootExternalDataLayerInstance ? RootExternalDataLayerInstance->GetExternalDataLayerAsset() : nullptr;
-	if (const UObject* ReferencingObject = RootExternalDataLayerAsset ? Cast<UObject>(RootExternalDataLayerAsset) : Cast<UObject>(DropTargetWorldDataLayers))
+	FText FailureReason;
+	if (!CanReferenceDataLayerAssets(DropTargetWorldDataLayers, DataLayerAssetsToDrop, &FailureReason))
 	{
-		for (const UDataLayerAsset* DataLayerAsset : DataLayerAssetsToDrop)
-		{
-			FText FailureReason;
-			if (!PassesAssetReferenceFiltering(ReferencingObject, DataLayerAsset, &FailureReason))
-			{	
-				return FSceneOutlinerDragValidationInfo(ESceneOutlinerDropCompatibility::Incompatible, FText::Format(LOCTEXT("CantCreateDataLayerInstancePassFilterFailed", "Cannot create Data Layer Instance : {0}"), FailureReason));
-			}
-		}
+		return FSceneOutlinerDragValidationInfo(ESceneOutlinerDropCompatibility::Incompatible, FText::Format(LOCTEXT("CantCreateDataLayerInstancePassFilterFailed", "Cannot create Data Layer Instance : {0}"), FailureReason));
 	}
 
 	// Check if target data layer supports having dropped asset as a child
@@ -1481,11 +1484,12 @@ void FDataLayerMode::RegisterContextMenu()
 					return CreateNewDataLayerInternal(nullptr, bInIsPrivate);
 				};
 
-				const AWorldDataLayers* WorldDataLayers = Mode->GetOwningWorld() ? Mode->GetOwningWorld()->GetWorldDataLayers() : nullptr;
-				if (WorldDataLayers && !WorldDataLayers->HasDeprecatedDataLayers())
+				const AWorldDataLayers* OwningWorldAWorldDataLayers = Mode->GetOwningWorldAWorldDataLayers();
+				if (OwningWorldAWorldDataLayers && !OwningWorldAWorldDataLayers->HasDeprecatedDataLayers())
 				{
+					const AWorldDataLayers* TargetWorldDataLayers = GetBestCandidateWorldDataLayersFromSelection(SceneOutliner, Mode);
 					Section.AddSubMenu("CreateNewDataLayerWithAsset", LOCTEXT("CreateNewDataLayerWithAssetSubMenu", "Create New Data Layer With Asset"), LOCTEXT("CreateNewDataLayerWithAssetSubMenu_ToolTip", "Create New Data Layer With Asset"),
-						FNewToolMenuDelegate::CreateLambda([CreateNewDataLayer, WorldDataLayers, Mode](UToolMenu* InSubMenu)
+						FNewToolMenuDelegate::CreateLambda([Mode, CreateNewDataLayer, OwningWorldAWorldDataLayers, TargetWorldDataLayers](UToolMenu* InSubMenu)
 						{
 							const bool bAllowClear = false;
 							const bool bAllowCopyPaste = false;
@@ -1498,14 +1502,30 @@ void FDataLayerMode::RegisterContextMenu()
 								bAllowCopyPaste,
 								AllowedClasses,
 								PropertyCustomizationHelpers::GetNewAssetFactoriesForClasses(AllowedClasses, NewAssetDisallowedClasses),
-								FOnShouldFilterAsset::CreateLambda([Mode, WorldDataLayers](const FAssetData& InAssetData)
+								FOnShouldFilterAsset::CreateLambda([Mode, OwningWorldAWorldDataLayers, TargetWorldDataLayers](const FAssetData& InAssetData)
 								{
-									// Filter already used Data Layers Assets and External Data Layer Assets that can't be added (those already added or not part of a registered GFD action)
-									const UExternalDataLayerManager* ExternalDataLayerManager = UExternalDataLayerManager::GetExternalDataLayerManager(Mode->GetOwningWorld());
 									const UDataLayerAsset* DataLayerAsset = Cast<UDataLayerAsset>(InAssetData.GetAsset());
 									const UExternalDataLayerAsset* ExternalDataLayerAsset = Cast<UExternalDataLayerAsset>(DataLayerAsset);
+
+									// Filter all External Data Layer Assets except if the target WorldDataLayers is the main one
+									if (ExternalDataLayerAsset && (TargetWorldDataLayers != OwningWorldAWorldDataLayers))
+									{
+										return true;
+									}
+
+									// Skip asset filtering for External Data Layer asset if target WorldDataLayers is the main one
+									if (!ExternalDataLayerAsset || (TargetWorldDataLayers != OwningWorldAWorldDataLayers))
+									{
+										if (!Mode->CanReferenceDataLayerAssets(TargetWorldDataLayers, { DataLayerAsset }))
+										{
+											return true;
+										}
+									}
+
+									// Filter already used Data Layers Assets and External Data Layer Assets that can't be added (those already added or not part of a registered GFD action)
+									const UExternalDataLayerManager* ExternalDataLayerManager = UExternalDataLayerManager::GetExternalDataLayerManager(Mode->GetOwningWorld());
 									const bool bCanInjectExternalDataLayerAsset = ExternalDataLayerAsset && ExternalDataLayerManager && ExternalDataLayerManager->CanInjectExternalDataLayerAsset(ExternalDataLayerAsset);
-									return !DataLayerAsset || (ExternalDataLayerAsset && !bCanInjectExternalDataLayerAsset) || WorldDataLayers->GetDataLayerInstance(DataLayerAsset);
+									return !DataLayerAsset || (ExternalDataLayerAsset && !bCanInjectExternalDataLayerAsset) || OwningWorldAWorldDataLayers->GetDataLayerInstance(DataLayerAsset);
 								}),
 								FOnAssetSelected::CreateLambda([CreateNewDataLayer](const FAssetData& InAssetData)
 								{
diff --git a/Engine/Source/Editor/DataLayerEditor/Private/DataLayer/DataLayerMode.h b/Engine/Source/Editor/DataLayerEditor/Private/DataLayer/DataLayerMode.h
--- a/Engine/Source/Editor/DataLayerEditor/Private/DataLayer/DataLayerMode.h
+++ b/Engine/Source/Editor/DataLayerEditor/Private/DataLayer/DataLayerMode.h
@@ -158,6 +158,7 @@ private:
 	UDataLayerInstance* GetDataLayerInstanceFromTreeItem(const ISceneOutlinerTreeItem& TreeItem) const;
 	FSceneOutlinerDragValidationInfo ValidateActorDrop(const ISceneOutlinerTreeItem& DropTarget, TArray<AActor*> PayloadActors, bool bMoveOperation = false) const;
 	FSceneOutlinerDragValidationInfo ValidateDataLayerAssetDrop(const ISceneOutlinerTreeItem& DropTarget, const TArray<const UDataLayerAsset*>& DataLayerAssetsToDrop) const;
+	bool CanReferenceDataLayerAssets(const AWorldDataLayers* InWorldDataLayers, const TArray<const UDataLayerAsset*>& InReferencedDataLayerAssets, FText* OutFailureReason = nullptr) const;
 	void OnDataLayerAssetDropped(const TArray<const UDataLayerAsset*>& DroppedDataLayerAsset, ISceneOutlinerTreeItem& DropTarget) const;
 
 	/** Filter factories */
diff --git a/Engine/Source/Runtime/Engine/Private/WorldPartition/DataLayer/WorldDataLayers.cpp b/Engine/Source/Runtime/Engine/Private/WorldPartition/DataLayer/WorldDataLayers.cpp
--- a/Engine/Source/Runtime/Engine/Private/WorldPartition/DataLayer/WorldDataLayers.cpp
+++ b/Engine/Source/Runtime/Engine/Private/WorldPartition/DataLayer/WorldDataLayers.cpp
@@ -1329,6 +1329,22 @@ void AWorldDataLayers::UpdateAccelerationTable(const UDataLayerInstance* DataLay
 #endif
 
 #if WITH_EDITOR
+
+bool AWorldDataLayers::CanReferenceDataLayerAsset(const UDataLayerAsset* InDataLayerAsset, FText* OutFailureReason) const
+{
+	auto PassesAssetReferenceFiltering = [](const UObject* InReferencingObject, const UDataLayerAsset* InDataLayerAsset, FText* OutReason)
+	{
+		FAssetReferenceFilterContext AssetReferenceFilterContext;
+		AssetReferenceFilterContext.ReferencingAssets.Add(FAssetData(InReferencingObject));
+		TSharedPtr<IAssetReferenceFilter> AssetReferenceFilter = GEditor->MakeAssetReferenceFilter(AssetReferenceFilterContext);
+		return AssetReferenceFilter.IsValid() ? AssetReferenceFilter->PassesFilter(FAssetData(InDataLayerAsset), OutReason) : true;
+	};
+
+	const UExternalDataLayerAsset* RootExternalDataLayerAsset = GetRootExternalDataLayerAsset();
+	const UObject* ReferencingObject = RootExternalDataLayerAsset ? Cast<UObject>(RootExternalDataLayerAsset) : Cast<UObject>(this);
+	return PassesAssetReferenceFiltering(ReferencingObject, InDataLayerAsset, OutFailureReason);
+}
+
 void AWorldDataLayers::RemoveEditorDataLayers()
 {
 	TArray<UDataLayerInstance*> EditorDataLayers;
diff --git a/Engine/Source/Runtime/Engine/Public/WorldPartition/DataLayer/WorldDataLayers.h b/Engine/Source/Runtime/Engine/Public/WorldPartition/DataLayer/WorldDataLayers.h
--- a/Engine/Source/Runtime/Engine/Public/WorldPartition/DataLayer/WorldDataLayers.h
+++ b/Engine/Source/Runtime/Engine/Public/WorldPartition/DataLayer/WorldDataLayers.h
@@ -111,6 +111,8 @@ public:
 	ENGINE_API bool SupportsExternalPackageDataLayerInstances() const;
 	ENGINE_API bool IsUsingExternalPackageDataLayerInstances() const { return bUseExternalPackageDataLayerInstances; }
 	ENGINE_API bool SetUseExternalPackageDataLayerInstances(bool bInNewValue, bool bInInteractiveMode = false);
+
+	ENGINE_API bool CanReferenceDataLayerAsset(const UDataLayerAsset* InDataLayerAsset, FText* OutFailureReason) const;
 #endif
 
 	UE_DEPRECATED(5.4, "Use GetWorldPartitionWorldDataLayersName() instead.")
