Data Layer Outliner improvements:
- Added Data Layer Debug Color column
- Now properly restores last outliner state when starting the editor (like visible columns)

diff --git a/Engine/Source/Editor/DataLayerEditor/Private/DataLayer/DataLayerOutlinerDebugColorColumn.cpp b/Engine/Source/Editor/DataLayerEditor/Private/DataLayer/DataLayerOutlinerDebugColorColumn.cpp
new file mode 100644
--- /dev/null
+++ b/Engine/Source/Editor/DataLayerEditor/Private/DataLayer/DataLayerOutlinerDebugColorColumn.cpp
@@ -0,0 +1,53 @@
+// Copyright Epic Games, Inc. All Rights Reserved.
+
+#include "DataLayerOutlinerDebugColorColumn.h"
+#include "WorldPartition/DataLayer/DataLayerInstance.h"
+#include "DataLayerTreeItem.h"
+#include "SceneOutlinerPublicTypes.h"
+#include "Types/SlateEnums.h"
+#include "Widgets/Images/SImage.h"
+#include "Widgets/Layout/SSpacer.h"
+#include "Widgets/SNullWidget.h"
+#include "Widgets/Views/SHeaderRow.h"
+
+#define LOCTEXT_NAMESPACE "DataLayer"
+
+FName FDataLayerOutlinerDebugColorColumn::GetID()
+{
+	static FName DataLayeDebugColor("Debug Color");
+	return DataLayeDebugColor;
+}
+
+SHeaderRow::FColumn::FArguments FDataLayerOutlinerDebugColorColumn::ConstructHeaderRowColumn()
+{
+	return SHeaderRow::Column(GetColumnID())
+		.FixedWidth(20.f)
+		.DefaultTooltip(FText::FromName(GetColumnID()))
+		[
+			SNew(SSpacer)
+		];
+}
+
+const TSharedRef<SWidget> FDataLayerOutlinerDebugColorColumn::ConstructRowWidget(FSceneOutlinerTreeItemRef TreeItem, const STableRow<FSceneOutlinerTreeItemPtr>& Row)
+{
+	if (TreeItem->IsA<FDataLayerTreeItem>())
+	{
+		return SNew(SBox)
+			.HAlign(HAlign_Center)
+			.VAlign(VAlign_Center)
+			.Content()
+			[
+				SNew(SImage)
+				.ColorAndOpacity_Lambda([TreeItem]()
+				{
+					FDataLayerTreeItem* DataLayerTreeItem = TreeItem->CastTo<FDataLayerTreeItem>();
+					UDataLayerInstance* DataLayerInstance = DataLayerTreeItem ? DataLayerTreeItem->GetDataLayer() : nullptr;
+					return DataLayerInstance ? DataLayerInstance->GetDebugColor() : FColor::Black;
+				})
+				.Image(FAppStyle::Get().GetBrush("Level.ColorIcon"))
+			];
+	}
+	return SNullWidget::NullWidget;
+}
+
+#undef LOCTEXT_NAMESPACE
\ No newline at end of file
diff --git a/Engine/Source/Editor/DataLayerEditor/Private/DataLayer/DataLayerOutlinerDebugColorColumn.h b/Engine/Source/Editor/DataLayerEditor/Private/DataLayer/DataLayerOutlinerDebugColorColumn.h
new file mode 100644
--- /dev/null
+++ b/Engine/Source/Editor/DataLayerEditor/Private/DataLayer/DataLayerOutlinerDebugColorColumn.h
@@ -0,0 +1,35 @@
+// Copyright Epic Games, Inc. All Rights Reserved.
+
+#pragma once
+
+#include "ISceneOutliner.h"
+#include "ISceneOutlinerColumn.h"
+#include "SceneOutlinerFwd.h"
+#include "Templates/SharedPointer.h"
+#include "Widgets/Views/SHeaderRow.h"
+#include "UObject/NameTypes.h"
+
+class SWidget;
+template <typename ItemType> class STableRow;
+
+class FDataLayerOutlinerDebugColorColumn : public ISceneOutlinerColumn
+{
+public:
+
+	FDataLayerOutlinerDebugColorColumn(ISceneOutliner& SceneOutliner) : WeakSceneOutliner(StaticCastSharedRef<ISceneOutliner>(SceneOutliner.AsShared())) {}
+	virtual ~FDataLayerOutlinerDebugColorColumn() {}
+	static FName GetID();
+
+	//////////////////////////////////////////////////////////////////////////
+	// Begin ISceneOutlinerColumn Implementation
+	virtual FName GetColumnID() override { return GetID(); }
+	virtual SHeaderRow::FColumn::FArguments ConstructHeaderRowColumn() override;
+	virtual const TSharedRef<SWidget> ConstructRowWidget(FSceneOutlinerTreeItemRef TreeItem, const STableRow<FSceneOutlinerTreeItemPtr>& Row) override;
+	virtual bool SupportsSorting() const override { return false; }
+	// End ISceneOutlinerColumn Implementation
+	//////////////////////////////////////////////////////////////////////////
+
+private:
+
+	TWeakPtr<ISceneOutliner> WeakSceneOutliner;
+};
\ No newline at end of file
diff --git a/Engine/Source/Editor/DataLayerEditor/Private/DataLayer/SDataLayerBrowser.cpp b/Engine/Source/Editor/DataLayerEditor/Private/DataLayer/SDataLayerBrowser.cpp
--- a/Engine/Source/Editor/DataLayerEditor/Private/DataLayer/SDataLayerBrowser.cpp
+++ b/Engine/Source/Editor/DataLayerEditor/Private/DataLayer/SDataLayerBrowser.cpp
@@ -9,6 +9,7 @@
 #include "DataLayerActorTreeItem.h"
 #include "DataLayerMode.h"
 #include "DataLayerOutlinerDeleteButtonColumn.h"
+#include "DataLayerOutlinerDebugColorColumn.h"
 #include "DataLayerOutlinerHasErrorColumn.h"
 #include "DataLayerOutlinerIsLoadedInEditorColumn.h"
 #include "DataLayerOutlinerIsVisibleColumn.h"
@@ -131,6 +132,7 @@ void SDataLayerBrowser::Construct(const FArguments& InArgs)
 		.Visibility_Lambda([]() { return UDataLayerEditorSubsystem::Get()->HasDeprecatedDataLayers() ? EVisibility::Visible : EVisibility::Collapsed; });
 
 	FSceneOutlinerInitializationOptions InitOptions;
+	InitOptions.OutlinerIdentifier = TEXT("DataLayerEditorOutliner");
 	InitOptions.bShowHeaderRow = true;
 	InitOptions.bShowParentTree = true;
 	InitOptions.bShowCreateNewFolder = false;
@@ -143,6 +145,7 @@ void SDataLayerBrowser::Construct(const FArguments& InArgs)
 	InitOptions.ColumnMap.Add("ID Name", FSceneOutlinerColumnInfo(ESceneOutlinerColumnVisibility::Invisible, 20, FCreateSceneOutlinerColumn::CreateStatic(&FTextInfoColumn::CreateTextInfoColumn, FName("ID Name"), InternalNameInfoText, FText::GetEmpty())));
 	InitOptions.ColumnMap.Add("Initial State", FSceneOutlinerColumnInfo(ESceneOutlinerColumnVisibility::Invisible, 20, FCreateSceneOutlinerColumn::CreateStatic(&FTextInfoColumn::CreateTextInfoColumn, FName("Initial State"), InternalInitialRuntimeStateInfoText, FText::FromString("Initial Runtime State"))));
 	InitOptions.ColumnMap.Add(FDataLayerOutlinerHasErrorsColumn::GetID(), FSceneOutlinerColumnInfo(ESceneOutlinerColumnVisibility::Visible, 100, FCreateSceneOutlinerColumn::CreateLambda([](ISceneOutliner& InSceneOutliner) { return MakeShareable(new FDataLayerOutlinerHasErrorsColumn(InSceneOutliner)); }), false));
+	InitOptions.ColumnMap.Add(FDataLayerOutlinerDebugColorColumn::GetID(), FSceneOutlinerColumnInfo(ESceneOutlinerColumnVisibility::Invisible, 100, FCreateSceneOutlinerColumn::CreateLambda([](ISceneOutliner& InSceneOutliner) { return MakeShareable(new FDataLayerOutlinerDebugColorColumn(InSceneOutliner)); })));
 	DataLayerOutliner = SNew(SDataLayerOutliner, InitOptions).IsEnabled(FSlateApplication::Get().GetNormalExecutionAttribute());
 
 	SAssignNew(DataLayerContentsSection, SBorder)
