Added 'Runtime Data Layer Color' and 'External Data Layer Color' Actor Coloration modes.

diff --git a/Engine/Source/Editor/DataLayerEditor/Private/DataLayer/DataLayerEditorSubsystem.cpp b/Engine/Source/Editor/DataLayerEditor/Private/DataLayer/DataLayerEditorSubsystem.cpp
--- a/Engine/Source/Editor/DataLayerEditor/Private/DataLayer/DataLayerEditorSubsystem.cpp
+++ b/Engine/Source/Editor/DataLayerEditor/Private/DataLayer/DataLayerEditorSubsystem.cpp
@@ -68,8 +68,9 @@ class SWidget;
 
 DEFINE_LOG_CATEGORY_STATIC(LogDataLayerEditorSubsystem, All, All);
 
-static FName NAME_DataLayerColor(TEXT("DataLayerColor"));
 static FName NAME_CurrentDataLayerColor(TEXT("CurrentDataLayerColor"));
+static FName NAME_RuntimeDataLayerColor(TEXT("RuntimeDataLayerColor"));
+static FName NAME_ExternalDataLayerColor(TEXT("ExternalDataLayerColor"));
 
 FDataLayerCreationParameters::FDataLayerCreationParameters()
 	: DataLayerAsset(nullptr)
@@ -163,6 +164,9 @@ void FDataLayersBroadcast::Initialize()
 		}
 
 #if ENABLE_ACTOR_PRIMITIVE_COLOR_HANDLER
+		// Colorize actor using its Data Layer Debug Color only if the Data Layer is in the Actor Editor Context
+		// - For multiple values, use white
+		// - Else, use gray
 		FActorPrimitiveColorHandler::Get().RegisterPrimitiveColorHandler(NAME_CurrentDataLayerColor, LOCTEXT("CurrentDataLayerColor", "Current Data Layer Color"), [](const UPrimitiveComponent* InPrimitiveComponent) -> FLinearColor
 		{
 			if (AActor* Actor = InPrimitiveComponent->GetOwner())
@@ -185,6 +189,61 @@ void FDataLayersBroadcast::Initialize()
 			}
 			return FLinearColor::Gray;
 		});
+
+		auto GetActorExternalDataLayerInstance = [](AActor* InActor) -> const UDataLayerInstance*
+		{
+			if (const UExternalDataLayerAsset* ExternalDataLayerAsset = InActor ? InActor->GetExternalDataLayerAsset() : nullptr)
+			{
+				if (const UDataLayerManager* DataLayerManager = UDataLayerManager::GetDataLayerManager(InActor))
+				{
+					if (const UDataLayerInstance* ExternalDataLayerInstance = DataLayerManager->GetDataLayerInstance(ExternalDataLayerAsset))
+					{
+						return ExternalDataLayerInstance;
+					}
+				}
+			}
+			return nullptr;
+		};
+
+		// Colorize actor using its Runtime Data Layer Debug Color
+		// - If 2 Runtime Data Layers and one of them is the EDL, favor the other one
+		// - Other cases of multiple Runtime Data Layers, use white
+		// - Else, use gray
+		FActorPrimitiveColorHandler::Get().RegisterPrimitiveColorHandler(NAME_RuntimeDataLayerColor, LOCTEXT("RuntimeDataLayerColor", "Runtime Data Layer Color"), [](const UPrimitiveComponent* InPrimitiveComponent) -> FLinearColor
+		{
+			if (AActor* Actor = InPrimitiveComponent->GetOwner())
+			{
+				TArray<const UDataLayerInstance*> RuntimeDataLayerInstances;
+				Algo::TransformIf(Actor->GetDataLayerInstances(), RuntimeDataLayerInstances, [](const UDataLayerInstance* DataLayerInstance) { return DataLayerInstance->IsRuntime(); }, [](const UDataLayerInstance* DataLayerInstance) { return DataLayerInstance; });
+				if (uint32 Count = RuntimeDataLayerInstances.Num())
+				{
+					if (Count == 1)
+					{
+						return RuntimeDataLayerInstances[0]->GetDebugColor();
+					}
+					else if (Count == 2)
+					{
+						if (RuntimeDataLayerInstances[0]->IsA<UExternalDataLayerInstance>())
+						{
+							return RuntimeDataLayerInstances[1]->GetDebugColor();
+						}
+						else if (RuntimeDataLayerInstances[1]->IsA<UExternalDataLayerInstance>())
+						{
+							return RuntimeDataLayerInstances[0]->GetDebugColor();
+						}
+					}
+					return FLinearColor::White;
+				}
+			}
+			return FLinearColor::Gray;
+		});
+
+		// Colorize actor using its External Data Layer Debug Color (Use gray if none)
+		FActorPrimitiveColorHandler::Get().RegisterPrimitiveColorHandler(NAME_ExternalDataLayerColor, LOCTEXT("ExternalDataLayerColor", "External Data Layer Color"), [&GetActorExternalDataLayerInstance](const UPrimitiveComponent* InPrimitiveComponent) -> FLinearColor
+		{
+			const UDataLayerInstance* ExternalDataLayerInstance = GetActorExternalDataLayerInstance(InPrimitiveComponent->GetOwner());
+			return ExternalDataLayerInstance ? ExternalDataLayerInstance->GetDebugColor() : FLinearColor::Gray;
+		});
 #endif
 	}
 }
@@ -298,7 +357,9 @@ void UDataLayerEditorSubsystem::Deinitialize()
 	DataLayerEditorLoadingStateChanged.Remove(OnActorDataLayersEditorLoadingStateChangedEngineBridgeHandle);
 
 #if ENABLE_ACTOR_PRIMITIVE_COLOR_HANDLER
+	FActorPrimitiveColorHandler::Get().UnregisterPrimitiveColorHandler(NAME_RuntimeDataLayerColor);
 	FActorPrimitiveColorHandler::Get().UnregisterPrimitiveColorHandler(NAME_CurrentDataLayerColor);
+	FActorPrimitiveColorHandler::Get().UnregisterPrimitiveColorHandler(NAME_ExternalDataLayerColor);
 #endif
 }
 
@@ -1576,7 +1637,9 @@ void UDataLayerEditorSubsystem::BroadcastDataLayerChanged(const EDataLayerAction
 #if ENABLE_ACTOR_PRIMITIVE_COLOR_HANDLER
 	if (UWorld* World = GetWorld())
 	{
+		FActorPrimitiveColorHandler::Get().RefreshPrimitiveColorHandler(NAME_RuntimeDataLayerColor, World);
 		FActorPrimitiveColorHandler::Get().RefreshPrimitiveColorHandler(NAME_CurrentDataLayerColor, World);
+		FActorPrimitiveColorHandler::Get().RefreshPrimitiveColorHandler(NAME_ExternalDataLayerColor, World);
 	}
 #endif
 }
