Exposed DataLayerInstance filter in FDataLayerPropertyTypeCustomizationHelper::CreateDataLayerMenu (for future use)
Fixed FOnShouldFilterDataLayerInstance so that true actually filters the Data Layer instance (implementation and test were inverted which was confusing).

diff --git a/Engine/Source/Editor/DataLayerEditor/Private/DataLayer/DataLayerMode.cpp b/Engine/Source/Editor/DataLayerEditor/Private/DataLayer/DataLayerMode.cpp
--- a/Engine/Source/Editor/DataLayerEditor/Private/DataLayer/DataLayerMode.cpp
+++ b/Engine/Source/Editor/DataLayerEditor/Private/DataLayer/DataLayerMode.cpp
@@ -1590,10 +1590,11 @@ void FDataLayerMode::RegisterContextMenu()
 									{
 										if (!DataLayerInstance->CanBeChildOf(InCandidateDataLayerInstance))
 										{
-											return false;
+											// True means should filter
+											return true;
 										}
 									}
-									return true;
+									return false;
 								}),
 								/*bShowRoot*/true);
 						}));
@@ -1657,10 +1658,11 @@ void FDataLayerMode::RegisterContextMenu()
 									{
 										if (!InCandidateDataLayerInstance->CanAddActor(const_cast<AActor*>(SelectedActor.Get())))
 										{
-											return false;
+											// True means should filter
+											return true;
 										}
 									}
-									return true;
+									return false;
 								})
 							);
 						}));
@@ -2371,7 +2373,7 @@ TSharedRef<SWidget> FDataLayerPickingMode::CreateDataLayerPickerWidget(FOnDataLa
 	InitOptions.bFocusSearchBoxWhenOpened = true;
 	if (OnShouldFilterDataLayerInstance.IsBound())
 	{
-		InitOptions.Filters->AddFilterPredicate<FDataLayerTreeItem>(FDataLayerTreeItem::FFilterPredicate::CreateLambda([OnShouldFilterDataLayerInstance](const UDataLayerInstance* DataLayerInstance) { return OnShouldFilterDataLayerInstance.Execute(DataLayerInstance); }), FSceneOutlinerFilter::EDefaultBehaviour::Pass);
+		InitOptions.Filters->AddFilterPredicate<FDataLayerTreeItem>(FDataLayerTreeItem::FFilterPredicate::CreateLambda([OnShouldFilterDataLayerInstance](const UDataLayerInstance* DataLayerInstance) { return !OnShouldFilterDataLayerInstance.Execute(DataLayerInstance); }), FSceneOutlinerFilter::EDefaultBehaviour::Pass);
 	}
 	InitOptions.ColumnMap.Add(FSceneOutlinerBuiltInColumnTypes::Label(), FSceneOutlinerColumnInfo(ESceneOutlinerColumnVisibility::Visible, 2));
 	InitOptions.ModeFactory = FCreateSceneOutlinerMode::CreateLambda([OnDataLayerInstancePicked](SSceneOutliner* Outliner)
diff --git a/Engine/Source/Editor/DataLayerEditor/Private/DataLayer/DataLayerPropertyTypeCustomizationHelper.cpp b/Engine/Source/Editor/DataLayerEditor/Private/DataLayer/DataLayerPropertyTypeCustomizationHelper.cpp
--- a/Engine/Source/Editor/DataLayerEditor/Private/DataLayer/DataLayerPropertyTypeCustomizationHelper.cpp
+++ b/Engine/Source/Editor/DataLayerEditor/Private/DataLayer/DataLayerPropertyTypeCustomizationHelper.cpp
@@ -9,12 +9,17 @@ class SWidget;
 
 #define LOCTEXT_NAMESPACE "DataLayer"
 
-TSharedRef<SWidget> FDataLayerPropertyTypeCustomizationHelper::CreateDataLayerMenu(TFunction<void(const UDataLayerInstance* DataLayer)> OnDataLayerSelectedFunction)
+TSharedRef<SWidget> FDataLayerPropertyTypeCustomizationHelper::CreateDataLayerMenu(TFunction<void(const UDataLayerInstance* DataLayerInstance)> OnDataLayerSelectedFunction, TFunction<bool(const UDataLayerInstance*)> OnShouldFilterDataLayerInstanceFunction)
 {
-	return FDataLayerPickingMode::CreateDataLayerPickerWidget(FOnDataLayerInstancePicked::CreateLambda([OnDataLayerSelectedFunction](UDataLayerInstance* TargetDataLayer)
-	{
-		OnDataLayerSelectedFunction(TargetDataLayer);
-	}));
+	return FDataLayerPickingMode::CreateDataLayerPickerWidget(
+		FOnDataLayerInstancePicked::CreateLambda([OnDataLayerSelectedFunction](UDataLayerInstance* TargetDataLayerInstance)
+		{
+			OnDataLayerSelectedFunction(TargetDataLayerInstance);
+		}),
+		FOnShouldFilterDataLayerInstance::CreateLambda([OnShouldFilterDataLayerInstanceFunction](const UDataLayerInstance* DataLayerInstance)
+		{
+			return OnShouldFilterDataLayerInstanceFunction(DataLayerInstance);
+		}));
 }
 
 #undef LOCTEXT_NAMESPACE
\ No newline at end of file
diff --git a/Engine/Source/Editor/DataLayerEditor/Public/DataLayer/DataLayerPropertyTypeCustomizationHelper.h b/Engine/Source/Editor/DataLayerEditor/Public/DataLayer/DataLayerPropertyTypeCustomizationHelper.h
--- a/Engine/Source/Editor/DataLayerEditor/Public/DataLayer/DataLayerPropertyTypeCustomizationHelper.h
+++ b/Engine/Source/Editor/DataLayerEditor/Public/DataLayer/DataLayerPropertyTypeCustomizationHelper.h
@@ -12,5 +12,5 @@ class UDataLayerInstance;
 
 struct DATALAYEREDITOR_API FDataLayerPropertyTypeCustomizationHelper
 {
-	static TSharedRef<SWidget> CreateDataLayerMenu(TFunction<void(const UDataLayerInstance* DataLayer)> OnDataLayerSelectedFunction);
+	static TSharedRef<SWidget> CreateDataLayerMenu(TFunction<void(const UDataLayerInstance* DataLayer)> OnDataLayerSelectedFunction, TFunction<bool(const UDataLayerInstance*)> OnShouldFilterDataLayerInstanceFunction = [](const UDataLayerInstance*){ return false; });
 };
