Fix CIS : Action graph contains cycle.

diff --git a/Engine/Source/Editor/DataLayerEditor/DataLayerEditor.Build.cs b/Engine/Source/Editor/DataLayerEditor/DataLayerEditor.Build.cs
--- a/Engine/Source/Editor/DataLayerEditor/DataLayerEditor.Build.cs
+++ b/Engine/Source/Editor/DataLayerEditor/DataLayerEditor.Build.cs
@@ -26,5 +26,7 @@ public class DataLayerEditor : ModuleRules
 				"ContentBrowserData"
 			}
 		);
+
+		CircularlyReferencedDependentModules.Add("LevelEditor");
 	}
 }
