Removed circular dependency of DataLayerEditor to LevelEditor.

diff --git a/Engine/Source/Editor/DataLayerEditor/DataLayerEditor.Build.cs b/Engine/Source/Editor/DataLayerEditor/DataLayerEditor.Build.cs
--- a/Engine/Source/Editor/DataLayerEditor/DataLayerEditor.Build.cs
+++ b/Engine/Source/Editor/DataLayerEditor/DataLayerEditor.Build.cs
@@ -15,7 +15,6 @@ public class DataLayerEditor : ModuleRules
 				"EditorSubsystem",
 				"PropertyEditor",
 				"Engine",
-				"LevelEditor",
 				"InputCore",
 				"Slate",
 				"SlateCore",
@@ -26,7 +25,5 @@ public class DataLayerEditor : ModuleRules
 				"ContentBrowserData"
 			}
 		);
-
-		CircularlyReferencedDependentModules.Add("LevelEditor");
 	}
 }
