Fixed DataLayerEditorSubsystem having duplicate registered world delegates launching the project with the World in parameter.

diff --git a/Engine/Source/Editor/DataLayerEditor/Private/DataLayer/DataLayerEditorSubsystem.cpp b/Engine/Source/Editor/DataLayerEditor/Private/DataLayer/DataLayerEditorSubsystem.cpp
--- a/Engine/Source/Editor/DataLayerEditor/Private/DataLayer/DataLayerEditorSubsystem.cpp
+++ b/Engine/Source/Editor/DataLayerEditor/Private/DataLayer/DataLayerEditorSubsystem.cpp
@@ -583,9 +583,21 @@ TArray<const UDataLayerInstance*> UDataLayerEditorSubsystem::GetDataLayerInstanc
 
 void UDataLayerEditorSubsystem::UpdateRegisteredWorldDelegates()
 {
+	if (UWorld* PreviousWorld = LastRegisteredWorldDelegates.Get())
+	{
+		PreviousWorld->RemoveOnActorPreSpawnInitialization(OnActorPreSpawnInitializationDelegate);
+		PreviousWorld->PersistentLevel->OnLoadedActorAddedToLevelEvent.RemoveAll(this);
+		PreviousWorld->OnWorldPartitionInitialized().RemoveAll(this);
+		PreviousWorld->OnWorldPartitionUninitialized().RemoveAll(this);
+	}
+
+	LastRegisteredWorldDelegates.Reset();
+	OnActorPreSpawnInitializationDelegate.Reset();
+
 	if (UWorld* World = GetWorld())
 	{
-		World->AddOnActorPreSpawnInitialization(FOnActorSpawned::FDelegate::CreateUObject(this, &UDataLayerEditorSubsystem::OnActorPreSpawnInitialization));
+		LastRegisteredWorldDelegates = World;
+		OnActorPreSpawnInitializationDelegate = World->AddOnActorPreSpawnInitialization(FOnActorSpawned::FDelegate::CreateUObject(this, &UDataLayerEditorSubsystem::OnActorPreSpawnInitialization));
 		World->PersistentLevel->OnLoadedActorAddedToLevelEvent.AddUObject(this, &UDataLayerEditorSubsystem::OnLoadedActorAddedToLevel);
 		World->OnWorldPartitionInitialized().AddUObject(this, &UDataLayerEditorSubsystem::OnWorldPartitionInitialized);
 		World->OnWorldPartitionUninitialized().AddUObject(this, &UDataLayerEditorSubsystem::OnWorldPartitionUninitialized);
diff --git a/Engine/Source/Editor/DataLayerEditor/Public/DataLayer/DataLayerEditorSubsystem.h b/Engine/Source/Editor/DataLayerEditor/Public/DataLayer/DataLayerEditorSubsystem.h
--- a/Engine/Source/Editor/DataLayerEditor/Public/DataLayer/DataLayerEditorSubsystem.h
+++ b/Engine/Source/Editor/DataLayerEditor/Public/DataLayer/DataLayerEditorSubsystem.h
@@ -796,6 +796,12 @@ private:
 	/** Delegate used to notify changes to ActorEditorContextSubsystem */
 	FOnActorEditorContextClientChanged ActorEditorContextClientChanged;
 
+	/** Last World to have registered world delegates */
+	TWeakObjectPtr<UWorld> LastRegisteredWorldDelegates;
+
+	/** Delegate handle for world's AddOnActorPreSpawnInitialization */
+	FDelegateHandle OnActorPreSpawnInitializationDelegate;
+
 	friend class FDataLayersBroadcast;
 	friend struct FExternalDataLayerWorldSurrogateReferencingObject;
 };
