World Partition: fix minor regression from 32133290 that was showing deleted actors as unloaded in the world outliner.

diff --git a/Engine/Source/Editor/DataLayerEditor/Private/DataLayer/DataLayerHierarchy.cpp b/Engine/Source/Editor/DataLayerEditor/Private/DataLayer/DataLayerHierarchy.cpp
--- a/Engine/Source/Editor/DataLayerEditor/Private/DataLayer/DataLayerHierarchy.cpp
+++ b/Engine/Source/Editor/DataLayerEditor/Private/DataLayer/DataLayerHierarchy.cpp
@@ -289,13 +289,22 @@ void FDataLayerHierarchy::CreateItems(TArray<FSceneOutlinerTreeItemPtr>& OutItem
 							return true;
 						}
 
-						const TSet<FGuid> LoadedActors =  FWorldPartitionHelpers::GetLoadedActorGuidsForLevel(OuterLevel);
+						const TMap<FGuid, AActor*> LoadedActors = FWorldPartitionHelpers::GetLoadedActorsForLevel(OuterLevel);
+						const TMap<FGuid, AActor*> RegisteredActors = FWorldPartitionHelpers::GetRegisteredActorsForLevel(OuterLevel);
 
 						// Create an FDataLayerActorDescTreeItem for each unloaded actor of this WorldPartition
-						FWorldPartitionHelpers::ForEachActorDescInstance(WorldPartition, [this, IsDataLayerShown, DataLayerManager, CurrentLevel, &WorldToLevelDataLayerMap, &LoadedActors, &OutItems](const FWorldPartitionActorDescInstance* ActorDescInstance)
+						FWorldPartitionHelpers::ForEachActorDescInstance(WorldPartition, [this, IsDataLayerShown, DataLayerManager, CurrentLevel, &WorldToLevelDataLayerMap, &LoadedActors, &RegisteredActors, &OutItems](const FWorldPartitionActorDescInstance* ActorDescInstance)
 						{
-							if (ActorDescInstance != nullptr && !LoadedActors.Contains(ActorDescInstance->GetGuid()) && FActorDescTreeItem::ShouldDisplayInOutliner(ActorDescInstance))
+							if (ActorDescInstance && FActorDescTreeItem::ShouldDisplayInOutliner(ActorDescInstance))
 							{
+								if (AActor* const* LoadedActor = LoadedActors.Find(ActorDescInstance->GetGuid()))
+								{
+									if (!IsValid(*LoadedActor) || RegisteredActors.Contains(ActorDescInstance->GetGuid()))
+									{
+										return true;
+									}
+								}
+
 								for (const FName& DataLayerInstanceName : ActorDescInstance->GetDataLayerInstanceNames().ToArray())
 								{
 									if (const UDataLayerInstance* DataLayerInstance = DataLayerManager->GetDataLayerInstance(DataLayerInstanceName))
diff --git a/Engine/Source/Editor/SceneOutliner/Private/ActorHierarchy.cpp b/Engine/Source/Editor/SceneOutliner/Private/ActorHierarchy.cpp
--- a/Engine/Source/Editor/SceneOutliner/Private/ActorHierarchy.cpp
+++ b/Engine/Source/Editor/SceneOutliner/Private/ActorHierarchy.cpp
@@ -519,12 +519,21 @@ void FActorHierarchy::CreateUnloadedItems(UWorld* World, TArray<FSceneOutlinerTr
 					}
 				}
 
-				const TSet<FGuid> LoadedActors = FWorldPartitionHelpers::GetLoadedActorGuidsForLevel(OuterLevel);
+				const TMap<FGuid, AActor*> LoadedActors = FWorldPartitionHelpers::GetLoadedActorsForLevel(OuterLevel);
+				const TMap<FGuid, AActor*> RegisteredActors = FWorldPartitionHelpers::GetRegisteredActorsForLevel(OuterLevel);
 
-				FWorldPartitionHelpers::ForEachActorDescInstance(WorldPartition, [this, &LoadedActors, &OutItems, &AddChildContainer, LevelInstanceSubsystem](const FWorldPartitionActorDescInstance* ActorDescInstance)
+				FWorldPartitionHelpers::ForEachActorDescInstance(WorldPartition, [this, &LoadedActors, &RegisteredActors, &OutItems, &AddChildContainer, LevelInstanceSubsystem](const FWorldPartitionActorDescInstance* ActorDescInstance)
 				{
-					if (ActorDescInstance != nullptr && !LoadedActors.Contains(ActorDescInstance->GetGuid()) && FActorDescTreeItem::ShouldDisplayInOutliner(ActorDescInstance))
+					if (ActorDescInstance && FActorDescTreeItem::ShouldDisplayInOutliner(ActorDescInstance))
 					{
+						if (AActor* const* LoadedActor = LoadedActors.Find(ActorDescInstance->GetGuid()))
+						{
+							if (!IsValid(*LoadedActor) || RegisteredActors.Contains(ActorDescInstance->GetGuid()))
+							{
+								return true;
+							}
+						}
+
 						UActorDescContainerInstance* ContainerInstance = ActorDescInstance->GetContainerInstance();
 
 						if (const FSceneOutlinerTreeItemPtr ActorDescItem = Mode->CreateItemFor<FActorDescTreeItem>(FActorDescTreeItem(ActorDescInstance->GetGuid(), ContainerInstance)))
diff --git a/Engine/Source/Runtime/Engine/Private/WorldPartition/WorldPartitionHelpers.cpp b/Engine/Source/Runtime/Engine/Private/WorldPartition/WorldPartitionHelpers.cpp
--- a/Engine/Source/Runtime/Engine/Private/WorldPartition/WorldPartitionHelpers.cpp
+++ b/Engine/Source/Runtime/Engine/Private/WorldPartition/WorldPartitionHelpers.cpp
@@ -324,10 +324,34 @@ bool FWorldPartitionHelpers::FixupRedirectedAssetPath(FName& InOutAssetPath)
 	return true;
 }
 
-TSet<FGuid> FWorldPartitionHelpers::GetLoadedActorGuidsForLevel(const ULevel* InLevel)
+TMap<FGuid, AActor*> FWorldPartitionHelpers::GetLoadedActorsForLevel(const ULevel* InLevel)
 {
-	TSet<FGuid> Result;
-	Algo::TransformIf(InLevel->Actors, Result, [](const AActor* Actor) { return IsValid(Actor); }, [](const AActor* Actor) { return Actor->GetActorGuid(); });
+	TMap<FGuid, AActor*> Result;
+	ForEachObjectWithOuter(InLevel, [&Result](UObject* Object)
+	{
+		if (AActor* Actor = Cast<AActor>(Object))
+		{
+			if (!Actor->IsTemplate() && Actor->GetActorGuid().IsValid())
+			{
+				Result.Add(Actor->GetActorGuid(), Actor);
+			}
+		}
+	});
+	return MoveTemp(Result);
+}
+
+TMap<FGuid, AActor*> FWorldPartitionHelpers::GetRegisteredActorsForLevel(const ULevel* InLevel)
+{
+	TMap<FGuid, AActor*> Result;
+	Algo::TransformIf(InLevel->Actors, Result, 
+		[](const AActor* Actor)
+		{
+			return IsValid(Actor);
+		}, 
+		[](AActor* Actor)
+		{
+			return TPair<FGuid, AActor*>(Actor->GetActorGuid(), Actor);
+		});
 	return MoveTemp(Result);
 }
 #endif // #if WITH_EDITOR
diff --git a/Engine/Source/Runtime/Engine/Public/WorldPartition/WorldPartitionHelpers.h b/Engine/Source/Runtime/Engine/Public/WorldPartition/WorldPartitionHelpers.h
--- a/Engine/Source/Runtime/Engine/Public/WorldPartition/WorldPartitionHelpers.h
+++ b/Engine/Source/Runtime/Engine/Public/WorldPartition/WorldPartitionHelpers.h
@@ -178,7 +178,11 @@ public:
 	UE_DEPRECATED(5.4, "Use UAssetRegistryHelpers::FixupRedirectedAssetPath instead.")
 	static ENGINE_API bool FixupRedirectedAssetPath(FSoftObjectPath& InOutSoftObjectPath);
 
-	static ENGINE_API TSet<FGuid> GetLoadedActorGuidsForLevel(const ULevel* InLevel);
+	// Returns all loaded actors for a specific level, including unregistered and garbage ones
+	static ENGINE_API TMap<FGuid, AActor*> GetLoadedActorsForLevel(const ULevel* InLevel);
+
+	// Returns all loaded and registered actors for a specific level
+	static ENGINE_API TMap<FGuid, AActor*> GetRegisteredActorsForLevel(const ULevel* InLevel);
 #endif // WITH_EDITOR
 
 	// Editor/Runtime conversions
