- Actor Editor Context current external data layer is now taken into account when duplicating actors.
- Duplicated actor is first initialized using source actor's External Data Layer. Then, Actor Editor Context tries to apply the current external data layer (which can fail if it would create invalid asset references).
- Also fixed UActorEditorContextSubsystem::ApplyContext when pasting actors : Actor Editor Context now waits at the end of the paste operation to apply the context as imported properties can override the properties if applied before (adapted Content Bundle to behave like before)

diff --git a/Engine/Plugins/Editor/WorldPartitionHLODUtilities/Source/Private/WorldPartition/HLOD/Utilities/WorldPartitionHLODUtilities.cpp b/Engine/Plugins/Editor/WorldPartitionHLODUtilities/Source/Private/WorldPartition/HLOD/Utilities/WorldPartitionHLODUtilities.cpp
--- a/Engine/Plugins/Editor/WorldPartitionHLODUtilities/Source/Private/WorldPartition/HLOD/Utilities/WorldPartitionHLODUtilities.cpp
+++ b/Engine/Plugins/Editor/WorldPartitionHLODUtilities/Source/Private/WorldPartition/HLOD/Utilities/WorldPartitionHLODUtilities.cpp
@@ -19,7 +19,6 @@
 #include "WorldPartition/HLOD/Builders/HLODBuilderMeshMerge.h"
 #include "WorldPartition/HLOD/Builders/HLODBuilderMeshSimplify.h"
 #include "WorldPartition/HLOD/Builders/HLODBuilderMeshApproximate.h"
-#include "ActorEditorContext/ScopedActorEditorContextSetExternalDataLayerAsset.h"
 
 #include "AssetCompilingManager.h"
 #include "BodySetupEnums.h"
@@ -205,7 +204,7 @@ TArray<AWorldPartitionHLOD*> FWorldPartitionHLODUtilities::CreateHLODActors(FHLO
 		if (bNewActor)
 		{
 			FContentBundleActivationScope Activationscope(InCreationParams.ContentBundleGuid);
-			FScopedActorEditorContextSetExternalDataLayerAsset EDLScope(InCreationParams.GetExternalDataLayerAsset());
+			FScopedOverrideSpawningLevelMountPointObject EDLScope(InCreationParams.GetExternalDataLayerAsset());
 
 			FActorSpawnParameters SpawnParams;
 			SpawnParams.Name = HLODActorName;
diff --git a/Engine/Source/Editor/DataLayerEditor/Private/DataLayer/DataLayerEditorSubsystem.cpp b/Engine/Source/Editor/DataLayerEditor/Private/DataLayer/DataLayerEditorSubsystem.cpp
--- a/Engine/Source/Editor/DataLayerEditor/Private/DataLayer/DataLayerEditorSubsystem.cpp
+++ b/Engine/Source/Editor/DataLayerEditor/Private/DataLayer/DataLayerEditorSubsystem.cpp
@@ -2,7 +2,6 @@
 
 #include "DataLayer/DataLayerEditorSubsystem.h"
 #include "DataLayer/DataLayerAction.h"
-#include "ActorEditorContext/ScopedActorEditorContextSetExternalDataLayerAsset.h"
 #include "Containers/EnumAsByte.h"
 #include "Containers/StringConv.h"
 #include "Containers/UnrealString.h"
@@ -31,6 +30,7 @@
 #include "Misc/Optional.h"
 #include "Misc/ScopedSlowTask.h"
 #include "Algo/Transform.h"
+#include "Algo/RemoveIf.h"
 #include "Algo/AnyOf.h"
 #include "ProfilingDebugging/CpuProfilerTrace.h"
 #include "Selection.h"
@@ -49,12 +49,17 @@
 #include "Widgets/Images/SImage.h"
 #include "Widgets/SBoxPanel.h"
 #include "Widgets/Text/STextBlock.h"
+#include "Widgets/Notifications/SNotificationList.h"
+#include "Framework/Notifications/NotificationManager.h"
+#include "Framework/Docking/TabManager.h"
+#include "DataLayerEditorModule.h"
 #include "WorldPartition/DataLayer/DataLayerAsset.h"
 #include "WorldPartition/DataLayer/DataLayerInstancePrivate.h"
 #include "WorldPartition/DataLayer/DataLayerInstanceWithAsset.h"
 #include "WorldPartition/DataLayer/DataLayerUtils.h"
 #include "WorldPartition/DataLayer/DataLayerManager.h"
 #include "WorldPartition/DataLayer/DeprecatedDataLayerInstance.h"
+#include "WorldPartition/DataLayer/ExternalDataLayerInstance.h"
 #include "WorldPartition/DataLayer/ExternalDataLayerAsset.h"
 #include "WorldPartition/DataLayer/ExternalDataLayerManager.h"
 #include "WorldPartition/DataLayer/ExternalDataLayerEngineSubsystem.h"
@@ -365,18 +370,28 @@ void UDataLayerEditorSubsystem::Deinitialize()
 
 void UDataLayerEditorSubsystem::OnActorPreSpawnInitialization(AActor* InActor)
 {
+	check(!InActor->GetExternalDataLayerAsset());
 	if (InActor->bIsEditorPreviewActor || !InActor->IsPackageExternal())
 	{
 		return;
 	}
 
-	UWorld* OwningWorld = InActor->GetWorld();
+	const UWorld* OwningWorld = InActor->GetWorld();
 	if (OwningWorld && (OwningWorld == GetWorld()))
 	{
-		if (const UExternalDataLayerAsset* ExternalDataLayerAsset = GetActorEditorContextCurrentExternalDataLayer())
+		// Prefer override spawning External Data Layer over actor editor context External Data Layer
+		const UExternalDataLayerAsset* OverrideSpawningExternalDataLayer = Cast<UExternalDataLayerAsset>(ULevel::GetOverrideSpawningLevelMountPointObject());
+		const UExternalDataLayerAsset* ExternalDataLayerAsset = OverrideSpawningExternalDataLayer ? OverrideSpawningExternalDataLayer : GetActorEditorContextCurrentExternalDataLayer();
+
+		UExternalDataLayerManager* ExternalDataLayerManager = UExternalDataLayerManager::GetExternalDataLayerManager(OwningWorld);
+		if (const UExternalDataLayerInstance* ExternalDataLayerInstance = ExternalDataLayerAsset ? ExternalDataLayerManager->GetExternalDataLayerInstance(ExternalDataLayerAsset) : nullptr)
 		{
-			UExternalDataLayerManager* ExternalDataLayerManager = UExternalDataLayerManager::GetExternalDataLayerManager(OwningWorld);
-			ExternalDataLayerManager->OnActorPreSpawnInitialization(InActor, ExternalDataLayerAsset);
+			FText FailureReason;
+			if (!FExternalDataLayerHelper::MoveActorsToExternalDataLayer({ InActor }, ExternalDataLayerInstance, &FailureReason))
+			{
+				UE_LOG(LogWorldPartition, Warning, TEXT("%s"), *FailureReason.ToString());
+				LastWarningNotification = FailureReason;
+			}
 		}
 	}
 }
@@ -393,7 +408,7 @@ ETickableTickType UDataLayerEditorSubsystem::GetTickableTickType() const
 
 bool UDataLayerEditorSubsystem::IsAllowedToTick() const
 {
-	return GetWorld() && (bAsyncBroadcastDataLayerChanged || bAsyncUpdateAllActorsVisibility || bAsyncInvalidateViewports);
+	return GetWorld() && (bAsyncBroadcastDataLayerChanged || bAsyncUpdateAllActorsVisibility || bAsyncInvalidateViewports || LastWarningNotification.IsSet());
 }
 
 void UDataLayerEditorSubsystem::Tick(float DeltaTime)
@@ -415,6 +430,17 @@ void UDataLayerEditorSubsystem::Tick(float DeltaTime)
 		GEditor->RedrawLevelEditingViewports();
 		bAsyncInvalidateViewports = false;
 	}
+
+	if (LastWarningNotification.IsSet())
+	{
+		// Trigger a notification with the last pushed warning (avoids spamming notification manager)
+		FNotificationInfo WarningInfo(LastWarningNotification.GetValue());
+		WarningInfo.ExpireDuration = 3.0f;
+		WarningInfo.Hyperlink = FSimpleDelegate::CreateLambda([]() { FGlobalTabmanager::Get()->TryInvokeTab(FName("OutputLog")); });
+		WarningInfo.HyperlinkText = LOCTEXT("ShowMessageLogHyperlink", "Show Output Log");
+		FSlateNotificationManager::Get().AddNotification(WarningInfo);
+		LastWarningNotification.Reset();
+	}
 }
 
 void UDataLayerEditorSubsystem::BeginDestroy()
@@ -438,7 +464,36 @@ void UDataLayerEditorSubsystem::OnExecuteActorEditorContextAction(UWorld* InWorl
 		case EActorEditorContextAction::ApplyContext:
 			check(InActor && InActor->GetWorld() == InWorld);
 			{
-				AddActorToDataLayers(InActor, DataLayerManager->GetActorEditorContextDataLayers());
+				// Try to apply context External Data Layer (this operation can fail if asset referencing validation fails)
+				// Don't apply if there's a valid override spawning External Data Layer (see OnActorPreSpawnInitialization)
+				const UExternalDataLayerAsset* OverrideSpawningExternalDataLayer = Cast<UExternalDataLayerAsset>(ULevel::GetOverrideSpawningLevelMountPointObject());
+				if (!OverrideSpawningExternalDataLayer)
+				{
+					const UExternalDataLayerAsset* CurrentExternalDataLayer = GetActorEditorContextCurrentExternalDataLayer();
+					const UExternalDataLayerAsset* ActorExternalDataLayerAsset = InActor->GetExternalDataLayerAsset();
+					if (CurrentExternalDataLayer && (ActorExternalDataLayerAsset != CurrentExternalDataLayer))
+					{
+						UExternalDataLayerManager* ExternalDataLayerManager = UExternalDataLayerManager::GetExternalDataLayerManager(InActor);
+						if (const UExternalDataLayerInstance* CurrentExternalDataLayerInstance = CurrentExternalDataLayer ? ExternalDataLayerManager->GetExternalDataLayerInstance(CurrentExternalDataLayer) : nullptr)
+						{
+							FText FailureReason;
+							if (!FExternalDataLayerHelper::MoveActorsToExternalDataLayer({ InActor }, CurrentExternalDataLayerInstance, &FailureReason))
+							{
+								UE_LOG(LogWorldPartition, Warning, TEXT("%s"), *FailureReason.ToString());
+								LastWarningNotification = FailureReason;
+							}
+						}
+					}
+				}
+
+				// Apply context Data Layers (except External Data Layer)
+				TArray<UDataLayerInstance*> DataLayerInstances = DataLayerManager->GetActorEditorContextDataLayers();
+				if (!DataLayerInstances.IsEmpty())
+				{
+					DataLayerInstances.SetNum(Algo::RemoveIf(DataLayerInstances, [](UDataLayerInstance* DataLayerInstance) { return DataLayerInstance->IsA<UExternalDataLayerInstance>(); }));
+					AddActorToDataLayers(InActor, DataLayerInstances);
+					InActor->FixupDataLayers();
+				}
 			}
 			break;
 		case EActorEditorContextAction::ResetContext:
@@ -1523,7 +1578,8 @@ UDataLayerInstance* UDataLayerEditorSubsystem::CreateDataLayerInstance(const FDa
 			{
 				NewDataLayerInstance = CreateDataLayerInstance<UDataLayerInstancePrivate>(WorldDataLayers);
 			}
-			else
+			// Don't create an instance if no valid asset is provided
+			else if (Parameters.DataLayerAsset)
 			{
 				if (UExternalDataLayerAsset* ExternalDataLayerAsset = Cast<UExternalDataLayerAsset>(Parameters.DataLayerAsset))
 				{
@@ -1712,8 +1768,8 @@ public:
 		UWorld* ReferencingWorld = CurrentLevelOuterWorld ? CurrentLevelOuterWorld : World;
 		
 		const UExternalDataLayerAsset* ExternalDataLayerAsset = UDataLayerEditorSubsystem::GetReferencingWorldSurrogateObjectForObject(ReferencingWorld, FSoftObjectPath(DroppedObjects[0]));
-		EDLContext = ExternalDataLayerAsset ? TUniquePtr<FScopedActorEditorContextSetExternalDataLayerAsset>(new FScopedActorEditorContextSetExternalDataLayerAsset(ExternalDataLayerAsset)) : nullptr;
-		return EDLContext.IsValid() && (UDataLayerEditorSubsystem::Get()->GetActorEditorContextCurrentExternalDataLayer() == ExternalDataLayerAsset);
+		EDLContext = ExternalDataLayerAsset ? TUniquePtr<FScopedOverrideSpawningLevelMountPointObject>(new FScopedOverrideSpawningLevelMountPointObject(ExternalDataLayerAsset)) : nullptr;
+		return EDLContext.IsValid() && (ULevel::GetOverrideSpawningLevelMountPointObject() == ExternalDataLayerAsset);
 	}
 
 	virtual bool OnPostDropObjects(UWorld* World, const TArray<UObject*>& DroppedObjects)
@@ -1727,7 +1783,7 @@ public:
 	}
 
 private:
-	TUniquePtr<FScopedActorEditorContextSetExternalDataLayerAsset> EDLContext;
+	TUniquePtr<FScopedOverrideSpawningLevelMountPointObject> EDLContext;
 };
 
 TUniquePtr<FLevelEditorDragDropWorldSurrogateReferencingObject> UDataLayerEditorSubsystem::OnLevelEditorDragDropWorldSurrogateReferencingObject(UWorld* ReferencingWorld, const FSoftObjectPath& Object)
diff --git a/Engine/Source/Editor/DataLayerEditor/Public/DataLayer/DataLayerEditorSubsystem.h b/Engine/Source/Editor/DataLayerEditor/Public/DataLayer/DataLayerEditorSubsystem.h
--- a/Engine/Source/Editor/DataLayerEditor/Public/DataLayer/DataLayerEditorSubsystem.h
+++ b/Engine/Source/Editor/DataLayerEditor/Public/DataLayer/DataLayerEditorSubsystem.h
@@ -15,6 +15,7 @@
 #include "Templates/SharedPointer.h"
 #include "Templates/UnrealTemplate.h"
 #include "Tickable.h"
+#include "Misc/Optional.h"
 #include "UObject/NameTypes.h"
 #include "UObject/ObjectMacros.h"
 #include "UObject/ObjectPtr.h"
@@ -802,6 +803,9 @@ private:
 	/** Delegate handle for world's AddOnActorPreSpawnInitialization */
 	FDelegateHandle OnActorPreSpawnInitializationDelegate;
 
+	/** Last pushed warning to be send to the notification manager at next tick */
+	TOptional<FText> LastWarningNotification;
+
 	friend class FDataLayersBroadcast;
 	friend struct FExternalDataLayerWorldSurrogateReferencingObject;
 };
diff --git a/Engine/Source/Editor/UnrealEd/Classes/Editor/UnrealEdEngine.h b/Engine/Source/Editor/UnrealEd/Classes/Editor/UnrealEdEngine.h
--- a/Engine/Source/Editor/UnrealEd/Classes/Editor/UnrealEdEngine.h
+++ b/Engine/Source/Editor/UnrealEd/Classes/Editor/UnrealEdEngine.h
@@ -865,6 +865,14 @@ public:
 	DECLARE_DELEGATE_RetVal(const TArray<FTemplateMapInfo>&, FGetTemplateMapInfos);
 	FGetTemplateMapInfos& OnGetTemplateMapInfos() { return GetTemplateMapInfosDelegate; }
 
+	/** Delegate called when UUnrealEdEngine::PasteActors starts executing */
+	DECLARE_MULTICAST_DELEGATE(FOnPasteActorsBegin);
+	FOnPasteActorsBegin& OnPasteActorsBegin() { return OnPasteActorsBeginDelegate; }
+
+	/** Delegate called when UUnrealEdEngine::PasteActors is done executing */
+	DECLARE_MULTICAST_DELEGATE_OneParam(FOnPasteActorsEnd, const TArray<AActor*>&)
+	FOnPasteActorsEnd& OnPasteActorsEnd() { return OnPasteActorsEndDelegate; }
+
 	/** Gets the canonical list of map templates that should be visible in new level picker. This function calls OnGetTemplateMapInfos to allow runtime override of the default maps */
 	UNREALED_API const TArray<FTemplateMapInfo>& GetTemplateMapInfos() const;
 
@@ -923,6 +931,12 @@ private:
 	/* Delegate to override TemplateMapInfos */
 	FGetTemplateMapInfos GetTemplateMapInfosDelegate;
 
+	/** Delegate called when UUnrealEdEngine::PasteActors starts executing */
+	FOnPasteActorsBegin OnPasteActorsBeginDelegate;
+
+	/** Delegate called when UUnrealEdEngine::PasteActors is done executing */
+	FOnPasteActorsEnd OnPasteActorsEndDelegate;
+
 	/** Transient unsaved version of template map infos used by the editor. */
 	TArray<FTemplateMapInfo> TemplateMapInfoCache;
 
diff --git a/Engine/Source/Editor/UnrealEd/Private/EditorActor.cpp b/Engine/Source/Editor/UnrealEd/Private/EditorActor.cpp
--- a/Engine/Source/Editor/UnrealEd/Private/EditorActor.cpp
+++ b/Engine/Source/Editor/UnrealEd/Private/EditorActor.cpp
@@ -324,6 +324,10 @@ void UUnrealEdEngine::PasteActors(TArray<AActor*>& OutPastedActors, UWorld* InWo
 		return;
 	}
 
+	OnPasteActorsBeginDelegate.Broadcast();
+
+	ON_SCOPE_EXIT{ OnPasteActorsEndDelegate.Broadcast(OutPastedActors); };
+
 	const FScopedBusyCursor BusyCursor;
 
 	FCachedActorLabels ActorLabels(InWorld);
diff --git a/Engine/Source/Editor/UnrealEd/Private/Factories/EditorFactories.cpp b/Engine/Source/Editor/UnrealEd/Private/Factories/EditorFactories.cpp
--- a/Engine/Source/Editor/UnrealEd/Private/Factories/EditorFactories.cpp
+++ b/Engine/Source/Editor/UnrealEd/Private/Factories/EditorFactories.cpp
@@ -929,8 +929,8 @@ UObject* ULevelFactory::FactoryCreateText
 
 				FString ExternalDataLayerAssetPathStr;
 				FParse::Value(Str, TEXT("ExternalDataLayerAsset="), ExternalDataLayerAssetPathStr);
-				UExternalDataLayerAsset* ExternalDataLayerAsset = Cast<UExternalDataLayerAsset>(FSoftObjectPath(ExternalDataLayerAssetPathStr).TryLoad());
-				FScopedActorEditorContextSetExternalDataLayerAsset EDLScope(ExternalDataLayerAsset);
+				const UExternalDataLayerAsset* ExternalDataLayerAsset = Cast<UExternalDataLayerAsset>(FSoftObjectPath(ExternalDataLayerAssetPathStr).TryLoad());
+				FScopedOverrideSpawningLevelMountPointObject ScopeOverrideLevelMountPoint(ExternalDataLayerAsset);
 
 				// If we're pasting from a class that belongs to a map we need to duplicate the class and use that instead
 				if (FBlueprintEditorUtils::IsAnonymousBlueprintClass(TempClass))
diff --git a/Engine/Source/Editor/UnrealEd/Private/Subsystems/ActorEditorContextSubsystem.cpp b/Engine/Source/Editor/UnrealEd/Private/Subsystems/ActorEditorContextSubsystem.cpp
--- a/Engine/Source/Editor/UnrealEd/Private/Subsystems/ActorEditorContextSubsystem.cpp
+++ b/Engine/Source/Editor/UnrealEd/Private/Subsystems/ActorEditorContextSubsystem.cpp
@@ -1,9 +1,11 @@
 // Copyright Epic Games, Inc. All Rights Reserved.
 
 #include "Subsystems/ActorEditorContextSubsystem.h"
+#include "Editor/UnrealEdEngine.h"
 #include "GameFramework/Actor.h"
 #include "IActorEditorContextClient.h"
 #include "ScopedTransaction.h"
+#include "UnrealEdGlobals.h"
 #include "Editor.h"
 
 #define LOCTEXT_NAMESPACE "ActorEditorContext"
@@ -19,15 +21,42 @@ void UActorEditorContextSubsystem::Initialize(FSubsystemCollectionBase& Collecti
 	Super::Initialize(Collection);
 
 	GEditor->OnLevelActorAdded().AddUObject(this, &UActorEditorContextSubsystem::ApplyContext);
+	if (GUnrealEd)
+	{
+		GUnrealEd->OnPasteActorsBegin().AddUObject(this, &UActorEditorContextSubsystem::OnPasteActorsBegin);
+		GUnrealEd->OnPasteActorsEnd().AddUObject(this, &UActorEditorContextSubsystem::OnPasteActorsEnd);
+	}
 }
 
 void UActorEditorContextSubsystem::Deinitialize()
 {
 	GEditor->OnLevelActorAdded().RemoveAll(this);
 
+	if (GUnrealEd)
+	{
+		GUnrealEd->OnPasteActorsBegin().RemoveAll(this);
+		GUnrealEd->OnPasteActorsEnd().RemoveAll(this);
+	}
+
 	Super::Deinitialize();
 }
 
+void UActorEditorContextSubsystem::OnPasteActorsBegin()
+{
+	// Disable ApplyContext while UUnrealEdEngine::PasteActors is executing as ImportObjectProperties is called after OnLevelActorAdded anyways
+	bIsApplyEnabled = false;
+}
+
+void UActorEditorContextSubsystem::OnPasteActorsEnd(const TArray<AActor*>& InActors)
+{
+	// Enable and run ApplyContext now that UUnrealEdEngine::PasteActors is done executing
+	bIsApplyEnabled = true;
+	for (AActor* Actor : InActors)
+	{
+		ApplyContext(Actor);
+	}
+}
+
 void UActorEditorContextSubsystem::RegisterClient(IActorEditorContextClient* Client)
 {
 	if (Client && !Clients.Contains(Client))
@@ -50,7 +79,7 @@ void UActorEditorContextSubsystem::UnregisterClient(IActorEditorContextClient* C
 
 void UActorEditorContextSubsystem::ApplyContext(AActor* InActor)
 {
-	if (GIsReinstancing)
+	if (GIsReinstancing || !bIsApplyEnabled)
 	{
 		return;
 	}
diff --git a/Engine/Source/Editor/UnrealEd/Public/Subsystems/ActorEditorContextSubsystem.h b/Engine/Source/Editor/UnrealEd/Public/Subsystems/ActorEditorContextSubsystem.h
--- a/Engine/Source/Editor/UnrealEd/Public/Subsystems/ActorEditorContextSubsystem.h
+++ b/Engine/Source/Editor/UnrealEd/Public/Subsystems/ActorEditorContextSubsystem.h
@@ -43,6 +43,10 @@ private:
 	UNREALED_API void OnActorEditorContextClientChanged(IActorEditorContextClient* Client);
 	UNREALED_API void ApplyContext(AActor* InActor);
 
+	void OnPasteActorsBegin();
+	void OnPasteActorsEnd(const TArray<AActor*>& InActors);
+
 	FOnActorEditorContextSubsystemChanged ActorEditorContextSubsystemChanged;
 	TArray<IActorEditorContextClient*> Clients;
+	bool bIsApplyEnabled = true;
 };
diff --git a/Engine/Source/Editor/WorldPartitionEditor/Private/WorldPartition/ContentBundle/ContentBundleEditorSubsystem.cpp b/Engine/Source/Editor/WorldPartitionEditor/Private/WorldPartition/ContentBundle/ContentBundleEditorSubsystem.cpp
--- a/Engine/Source/Editor/WorldPartitionEditor/Private/WorldPartition/ContentBundle/ContentBundleEditorSubsystem.cpp
+++ b/Engine/Source/Editor/WorldPartitionEditor/Private/WorldPartition/ContentBundle/ContentBundleEditorSubsystem.cpp
@@ -5,6 +5,9 @@
 #include "WorldPartition/ContentBundle/ContentBundleDescriptor.h"
 #include "WorldPartition/ContentBundle/ContentBundle.h"
 #include "WorldPartition/ContentBundle/ContentBundleEditor.h"
+#include "WorldPartition/DataLayer/ExternalDataLayerManager.h"
+#include "WorldPartition/DataLayer/ExternalDataLayerInstance.h"
+#include "WorldPartition/DataLayer/ExternalDataLayerAsset.h"
 #include "WorldPartition/WorldPartition.h"
 #include "Subsystems/ActorEditorContextSubsystem.h"
 #include "Engine/Selection.h"
@@ -20,6 +23,8 @@ void UContentBundleEditingSubmodule::DoInitialize()
 {
 	check(GEditor);
 	UActorEditorContextSubsystem::Get()->RegisterClient(this);
+	// For backward compatibility, Content Bundle will be applied right when LevelActorAdded is called
+	GEditor->OnLevelActorAdded().AddUObject(this, &UContentBundleEditingSubmodule::ApplyContext);
 }
 
 void UContentBundleEditingSubmodule::DoDenitialize()
@@ -27,6 +32,7 @@ void UContentBundleEditingSubmodule::DoDenitialize()
 	if (GEditor)
 	{
 		UActorEditorContextSubsystem::Get()->UnregisterClient(this);
+		GEditor->OnLevelActorAdded().RemoveAll(this);
 	}
 
 	EditingContentBundleGuid.Invalidate();
@@ -66,15 +72,33 @@ void UContentBundleEditingSubmodule::PostEditUndo()
 	PreUndoRedoEditingContentBundleGuid.Invalidate();
 }
 
+void UContentBundleEditingSubmodule::ApplyContext(AActor* InActor)
+{
+	if (GIsReinstancing || !InActor || (InActor->GetWorld() != GetWorld()) || InActor->HasAnyFlags(RF_Transient) || InActor->IsChildActor() || !EditingContentBundleGuid.IsValid() || (InActor->GetContentBundleGuid() == EditingContentBundleGuid))
+	{
+		return;
+	}
+
+	// Prefer override spawning External Data Layer over Content Bundles
+	const UExternalDataLayerAsset* ExternalDataLayerAsset = Cast<UExternalDataLayerAsset>(ULevel::GetOverrideSpawningLevelMountPointObject());
+	const UExternalDataLayerManager* ExternalDataLayerManager = UExternalDataLayerManager::GetExternalDataLayerManager(InActor);
+	if (const UExternalDataLayerInstance* OverrideSpawningExternalDataLayerInstance = (ExternalDataLayerAsset && ExternalDataLayerManager) ? ExternalDataLayerManager->GetExternalDataLayerInstance(ExternalDataLayerAsset) : nullptr)
+	{
+		return;
+	}
+
+	if (TSharedPtr<FContentBundleEditor> EditingContentBundle = GetEditorContentBundle(EditingContentBundleGuid))
+	{
+		EditingContentBundle->AddActor(InActor);
+	}
+}
+
 void UContentBundleEditingSubmodule::OnExecuteActorEditorContextAction(UWorld* InWorld, const EActorEditorContextAction& InType, AActor* InActor /* = nullptr */)
 {
 	switch (InType)
 	{
 	case EActorEditorContextAction::ApplyContext:
-		if (TSharedPtr<FContentBundleEditor> EditingContentBundle = GetEditorContentBundle(EditingContentBundleGuid))
-		{
-			EditingContentBundle->AddActor(InActor);
-		}
+		ApplyContext(InActor);
 		break;
 	case EActorEditorContextAction::ResetContext:
 		if (TSharedPtr<FContentBundleEditor> EditingContentBundle = GetEditorContentBundle(EditingContentBundleGuid))
diff --git a/Engine/Source/Editor/WorldPartitionEditor/Public/WorldPartition/ContentBundle/ContentBundleEditorSubsystem.h b/Engine/Source/Editor/WorldPartitionEditor/Public/WorldPartition/ContentBundle/ContentBundleEditorSubsystem.h
--- a/Engine/Source/Editor/WorldPartitionEditor/Public/WorldPartition/ContentBundle/ContentBundleEditorSubsystem.h
+++ b/Engine/Source/Editor/WorldPartitionEditor/Public/WorldPartition/ContentBundle/ContentBundleEditorSubsystem.h
@@ -74,6 +74,7 @@ private:
 	void PopContentBundleEditing();
 	void StartEditing(TSharedPtr<FContentBundleEditor>& ContentBundleEditor);
 	void StopEditing(TSharedPtr<FContentBundleEditor>& ContentBundleEditor);
+	void ApplyContext(AActor* InActor);
 	virtual void PreEditUndo() override;
 	virtual void PostEditUndo() override;
 
diff --git a/Engine/Source/Runtime/Engine/Classes/Engine/Level.h b/Engine/Source/Runtime/Engine/Classes/Engine/Level.h
--- a/Engine/Source/Runtime/Engine/Classes/Engine/Level.h
+++ b/Engine/Source/Runtime/Engine/Classes/Engine/Level.h
@@ -1374,6 +1374,8 @@ public:
 	static ENGINE_API FDelegateHandle RegisterLevelMountPointResolver(const FLevelMountPointResolverDelegate& Resolver);
 	/** Unregisters a level mount point resolver */
 	static ENGINE_API void UnregisterLevelMountPointResolver(const FDelegateHandle& ResolverDelegateHandle);
+	/** Returns override spawning level mount point object */
+	static ENGINE_API const UObject* GetOverrideSpawningLevelMountPointObject() { return ULevel::OverrideSpawningLevelMountPointObject.Get(); }
 #endif
 
 private:
@@ -1387,6 +1389,9 @@ private:
 
 	/** Array of registered delegates used by GetExternalActorsPaths. */
 	static TArray<FLevelMountPointResolverDelegate> LevelMountPointResolvers;
+
+	/** Override spawning Level Mount Point object used when spawning an actor using external packaging */
+	static ENGINE_API TWeakObjectPtr<const UObject> OverrideSpawningLevelMountPointObject;
 private:
 	/**
 	 * Potentially defer the running of an actor's construction script on load
@@ -1414,6 +1419,7 @@ private:
 	/** Returns unreferenced actor folders that are marked as deleted. */
 	TSet<FGuid> GetDeletedAndUnreferencedActorFolders() const;
 
+	friend struct FScopedOverrideSpawningLevelMountPointObject;
 	friend struct FLevelActorFoldersHelper;
 	friend struct FSetWorldPartitionRuntimeCell;
 	friend class FWorldPartitionLevelHelper;
@@ -1435,6 +1441,21 @@ private:
 };
 
 #if WITH_EDITOR
+struct ENGINE_API FScopedOverrideSpawningLevelMountPointObject
+{
+	FScopedOverrideSpawningLevelMountPointObject(const UObject* InObject)
+	{
+		PreviousValue = ULevel::OverrideSpawningLevelMountPointObject;
+		ULevel::OverrideSpawningLevelMountPointObject = InObject;
+	}
+	~FScopedOverrideSpawningLevelMountPointObject()
+	{
+		ULevel::OverrideSpawningLevelMountPointObject = PreviousValue;
+	}
+private:
+	TWeakObjectPtr<const UObject> PreviousValue;
+};
+
 struct FSetWorldPartitionRuntimeCell
 {
 private:
diff --git a/Engine/Source/Runtime/Engine/Classes/GameFramework/Actor.h b/Engine/Source/Runtime/Engine/Classes/GameFramework/Actor.h
--- a/Engine/Source/Runtime/Engine/Classes/GameFramework/Actor.h
+++ b/Engine/Source/Runtime/Engine/Classes/GameFramework/Actor.h
@@ -4227,7 +4227,6 @@ private:
 	friend struct FSetActorInstanceGuid;
 	friend struct FSetActorContentBundleGuid;
 	friend struct FAssignActorDataLayer;
-	friend struct FScopeRawAssignActorExternalDataLayer;
 	friend struct FSetActorSelectable;
 	friend struct FSetActorFolderPath;
 #endif
diff --git a/Engine/Source/Runtime/Engine/Private/ActorEditor.cpp b/Engine/Source/Runtime/Engine/Private/ActorEditor.cpp
--- a/Engine/Source/Runtime/Engine/Private/ActorEditor.cpp
+++ b/Engine/Source/Runtime/Engine/Private/ActorEditor.cpp
@@ -1740,11 +1740,25 @@ void AActor::FixupDataLayers(bool bRevertChangesOnLockedDataLayer /*= false*/)
 	{ 
 		return;
 	}
+
+	// Remove any data layer with a Root external data layer different from the actor's external data layer
+	// This case can happen after duplicating an actor using the actor editor context with a different external data layer
+	TSet<TSoftObjectPtr<const UDataLayerAsset>> InvalidDataLayerAssets;
+	if (ExternalDataLayerAsset && !DataLayerAssets.IsEmpty())
+	{
+		for (const UDataLayerInstance* DataLayerInstance : GetDataLayerInstances())
+		{
+			if (const UExternalDataLayerInstance* RootExternalDataLayerInstance = DataLayerInstance->GetRootExternalDataLayerInstance(); RootExternalDataLayerInstance && (RootExternalDataLayerInstance->GetAsset() != ExternalDataLayerAsset))
+			{
+				InvalidDataLayerAssets.Add(DataLayerInstance->GetAsset());
+			}
+		}
+	}
 	
 	// Cleanup Data Layer assets we can't reference (Private DLs)
-	DataLayerAssets.SetNum(Algo::RemoveIf(DataLayerAssets, [this](const TSoftObjectPtr<UDataLayerAsset>& AssetPath)
+	DataLayerAssets.SetNum(Algo::RemoveIf(DataLayerAssets, [this, &InvalidDataLayerAssets](const TSoftObjectPtr<UDataLayerAsset>& AssetPath)
 	{
-		return !UDataLayerAsset::CanBeReferencedByActor(AssetPath, this);
+		return InvalidDataLayerAssets.Contains(AssetPath) || !UDataLayerAsset::CanBeReferencedByActor(AssetPath, this);
 	}));
 
 	// Use Actor's DataLayerManager since the fixup is relative to this level
diff --git a/Engine/Source/Runtime/Engine/Private/Level.cpp b/Engine/Source/Runtime/Engine/Private/Level.cpp
--- a/Engine/Source/Runtime/Engine/Private/Level.cpp
+++ b/Engine/Source/Runtime/Engine/Private/Level.cpp
@@ -118,6 +118,7 @@ static FAutoConsoleVariableRef CVarRouteActorInitializationWorkUnitWeighting(
 
 TArray<ULevel::FLevelExternalActorsPathsProviderDelegate> ULevel::LevelExternalActorsPathsProviders;
 TArray<ULevel::FLevelMountPointResolverDelegate> ULevel::LevelMountPointResolvers;
+TWeakObjectPtr<const UObject> ULevel::OverrideSpawningLevelMountPointObject;
 
 void FLevelActorFoldersHelper::SetUseActorFolders(ULevel* InLevel, bool bInEnabled)
 {
diff --git a/Engine/Source/Runtime/Engine/Private/LevelInstance/LevelInstanceSubsystem.cpp b/Engine/Source/Runtime/Engine/Private/LevelInstance/LevelInstanceSubsystem.cpp
--- a/Engine/Source/Runtime/Engine/Private/LevelInstance/LevelInstanceSubsystem.cpp
+++ b/Engine/Source/Runtime/Engine/Private/LevelInstance/LevelInstanceSubsystem.cpp
@@ -15,7 +15,6 @@
 
 #if WITH_EDITOR
 #include "Settings/LevelEditorMiscSettings.h"
-#include "ActorEditorContext/ScopedActorEditorContextSetExternalDataLayerAsset.h"
 #include "LevelInstance/LevelInstanceEditorLevelStreaming.h"
 #include "LevelInstance/LevelInstanceEditorPropertyOverrideLevelStreaming.h"
 #include "LevelInstance/ILevelInstanceEditorModule.h"
@@ -1462,7 +1461,7 @@ ILevelInstanceInterface* ULevelInstanceSubsystem::CreateLevelInstanceFrom(const
 	{
 		ExternalDataLayerInstance->bSkipCheckReadOnlyForSubLevels = true;
 	}
-	FScopedActorEditorContextSetExternalDataLayerAsset EDLScope(ExternalDataLayerAsset);
+	FScopedOverrideSpawningLevelMountPointObject EDLScope(ExternalDataLayerAsset);
 
 	if (!ActorClass->IsChildOf<APackedLevelActor>())
 	{
@@ -1878,6 +1877,15 @@ bool ULevelInstanceSubsystem::CanMoveActorToLevel(const AActor* Actor, FText* Ou
 		return false;
 	}
 
+	if (Actor->GetExternalDataLayerAsset())
+	{
+		if (OutReason != nullptr)
+		{
+			*OutReason = LOCTEXT("CantMoveActorUsingExternalDataLayer", "Can't move Level Instance actor using External Data Layer");
+		}
+		return false;
+	}
+
 	if (Actor->GetWorld() == GetWorld())
 	{
 		if (const ILevelInstanceInterface* LevelInstance = Cast<ILevelInstanceInterface>(Actor))
diff --git a/Engine/Source/Runtime/Engine/Private/WorldPartition/DataLayer/ExternalDataLayerHelper.cpp b/Engine/Source/Runtime/Engine/Private/WorldPartition/DataLayer/ExternalDataLayerHelper.cpp
--- a/Engine/Source/Runtime/Engine/Private/WorldPartition/DataLayer/ExternalDataLayerHelper.cpp
+++ b/Engine/Source/Runtime/Engine/Private/WorldPartition/DataLayer/ExternalDataLayerHelper.cpp
@@ -103,20 +103,15 @@ void FExternalDataLayerHelper::GetExternalDataLayerUIDs(const FAssetData& Asset,
 
 namespace UE::Private::ExternalDataLayerHelper
 {
-	static bool ValidateAssetUsingAssetReferenceRestrictions(const UObject* InAsset, TSet<FString>& OutInvalidReferenceReasons)
+	static bool ValidateAssetUsingAssetReferenceRestrictions(const UObject* InAsset, const TSet<UObject*>& InReferencedAssets, TSet<FString>& OutInvalidReferenceReasons)
 	{
 		if (!InAsset)
 		{
 			return false;
 		}
 
-		TArray<UClass*> IgnoreClasses;
-		TArray<UPackage*> IgnorePackages;
-		TSet<UObject*> ReferencedAssets;
-		FFindReferencedAssets::BuildAssetList(const_cast<UObject*>(InAsset), IgnoreClasses, IgnorePackages, ReferencedAssets);
-
 		uint32 ErrorCount = 0;
-		if (ReferencedAssets.Num() > 0)
+		if (InReferencedAssets.Num() > 0)
 		{
 			FAssetData ReferencingAssetData = FAssetData(InAsset);
 			FAssetReferenceFilterContext AssetReferenceFilterContext;
@@ -124,7 +119,7 @@ namespace UE::Private::ExternalDataLayerHelper
 			TSharedPtr<IAssetReferenceFilter> AssetReferenceFilter = GEditor ? GEditor->MakeAssetReferenceFilter(AssetReferenceFilterContext) : nullptr;
 			if (ensure(AssetReferenceFilter.IsValid()))
 			{
-				for (UObject* ReferencedAsset : ReferencedAssets)
+				for (UObject* ReferencedAsset : InReferencedAssets)
 				{
 					FText FailureReason;
 					if (!AssetReferenceFilter->PassesFilter(FAssetData(ReferencedAsset), &FailureReason))
@@ -139,61 +134,48 @@ namespace UE::Private::ExternalDataLayerHelper
 	}
 }
 
-// Only meant to be used by FExternalDataLayerHelper on stack
-struct FScopeRawAssignActorExternalDataLayer
-{
-private:
-	FScopeRawAssignActorExternalDataLayer(AActor* InActor, const UExternalDataLayerAsset* InExternalDataLayerAsset)
-		: Actor(InActor)
-		, ExternalDataLayerAsset(InActor->ExternalDataLayerAsset)
-	{
-		check(Actor);
-		Actor->ExternalDataLayerAsset = InExternalDataLayerAsset;
-	}
-
-	~FScopeRawAssignActorExternalDataLayer()
-	{
-		Actor->ExternalDataLayerAsset = ExternalDataLayerAsset;
-	}
-
-	AActor* Actor;
-	const UExternalDataLayerAsset* ExternalDataLayerAsset;
-
-	friend class FExternalDataLayerHelper;
-};
-
 bool FExternalDataLayerHelper::CanMoveActorsToExternalDataLayer(const TArray<AActor*>& InActors, const UExternalDataLayerInstance* InExternalDataLayerInstance, FText* OutFailureReason)
 {
 	auto CanMoveActorToExternalDataLayer = [](AActor* InActor, const UExternalDataLayerInstance* InExternalDataLayerInstance, FText& OutFailureReason)
 	{
+		const TArray<UClass*> IgnoreClasses;
+		const TArray<UPackage*> IgnorePackages;
+
 		check(!InActor->IsTemplate());
 		check(InActor->GetLevel());
 
+		const UExternalDataLayerAsset* OldExternalDataLayerAsset = InActor->GetExternalDataLayerAsset();
 		const UExternalDataLayerAsset* NewExternalDataLayerAsset = InExternalDataLayerInstance ? InExternalDataLayerInstance->GetExternalDataLayerAsset() : nullptr;
 		if (!InActor->IsPackageExternal())
 		{
 			OutFailureReason = FText::Format(LOCTEXT("CantMoveActorToEDL_NotPackageExternal", "Actor {0} is not using external package."), FText::FromString(InActor->GetName()));
 			return false;
 		}
-		
-		if (!InActor->IsUserManaged())
+
+		if (!InActor->IsUserManaged() && !InActor->GetExternalPackage()->HasAnyPackageFlags(PKG_NewlyCreated))
 		{
 			OutFailureReason = FText::Format(LOCTEXT("CantMoveActorToEDL_NotUserManaged", "Actor {0} cannot be manually modified."), FText::FromString(InActor->GetName()));
 			return false;
 		}
 
-		if (!InActor->GetExternalDataLayerAsset() && !NewExternalDataLayerAsset)
+		if (!OldExternalDataLayerAsset && !NewExternalDataLayerAsset)
 		{
 			OutFailureReason = FText::Format(LOCTEXT("CantMoveActorToEDL_NoExternalDataLayer", "Actor {0} has already no External Data Layer."), FText::FromString(InActor->GetName()));
 			return false;
 		}
 
-		if (InActor->GetExternalDataLayerAsset() == NewExternalDataLayerAsset)
+		if (OldExternalDataLayerAsset == NewExternalDataLayerAsset)
 		{
 			OutFailureReason = FText::Format(LOCTEXT("CantMoveActorToEDL_SameExternalDataLayer", "Actor {0} is already assigned to this External Data Layer."), FText::FromString(InActor->GetName()));
 			return false;
 		}
 
+		if (NewExternalDataLayerAsset && !InActor->SupportsDataLayerType(UExternalDataLayerInstance::StaticClass()))
+		{
+			OutFailureReason = FText::Format(LOCTEXT("CantMoveActorToEDL_EDLNotSupported", "Actor {0} doesn't support External Data Layers."), FText::FromString(InActor->GetName()));
+			return false;
+		}
+
 		if (InExternalDataLayerInstance && InExternalDataLayerInstance->IsReadOnly())
 		{
 			OutFailureReason = FText::Format(LOCTEXT("CantMoveActorToEDL_ReadOnlyExternalDataLayer", "External Data Layer is read-only."), FText::FromString(InExternalDataLayerInstance->GetDataLayerShortName()));
@@ -207,27 +189,26 @@ bool FExternalDataLayerHelper::CanMoveActorsToExternalDataLayer(const TArray<AAc
 			return false;
 		}
 
-		// Create a temporary transient copy of the actor
-		AActor* DuplicatedActorToValidate = nullptr;
+		// Gather actor asset references
+		TSet<UObject*> ActorReferencedAssets;
+		FFindReferencedAssets::BuildAssetList(InActor, IgnoreClasses, IgnorePackages, ActorReferencedAssets, true);
+		TArray<UObject*> ReferencedContent;
+		InActor->GetReferencedContentObjects(ReferencedContent);
+		// Remove itself and its data layer assets from the list
+		ActorReferencedAssets.Append(ReferencedContent);
+		ActorReferencedAssets.Remove(InActor);
+		for (auto It = ActorReferencedAssets.CreateIterator(); It; ++It)
 		{
-			// Avoid triggering any dirty package notifications as this is a transient object
-			TGuardValue<bool> IsEditorLoadingPackageGuard(GIsEditorLoadingPackage, true);
+			if (!(*It)->IsAsset() || (*It)->IsA<UDataLayerAsset>())
 			{
-				UWorld* World = InActor->GetWorld();
-				// Temporarily set actor to target value for the duplication to use this value
-				FScopeRawAssignActorExternalDataLayer ScopeSetActorExternalDataLayer(InActor, NewExternalDataLayerAsset);
-				// Flag the temporary actor transient to make sure IsAsset returns false (to avoid considering this temporary actor by any other systems)
-				FDelegateHandle PreSpawnDelegateHandle = World->AddOnActorPreSpawnInitialization(FOnActorSpawned::FDelegate::CreateLambda([](AActor* DuplicatedActorToValidate) { DuplicatedActorToValidate->SetFlags(RF_Transient); }));
-				DuplicatedActorToValidate = EditorActorSubsystem->DuplicateActor(InActor);
-				World->RemoveOnActorPreSpawnInitialization(PreSpawnDelegateHandle);
+				It.RemoveCurrent();
 			}
 		}
-		// Destroy the temporary actor
-		ON_SCOPE_EXIT{ DuplicatedActorToValidate->GetWorld()->EditorDestroyActor(DuplicatedActorToValidate, false); };
 
-		// Validate the temporary actor's asset references
+		// Validate if there are restrictions between the world or the new data layer asset and the actor asset references
 		TSet<FString> InvalidReferenceReasons;
-		if (!UE::Private::ExternalDataLayerHelper::ValidateAssetUsingAssetReferenceRestrictions(DuplicatedActorToValidate, InvalidReferenceReasons))
+		const UObject* Referencer = NewExternalDataLayerAsset ? (UObject*)NewExternalDataLayerAsset : (UObject*)InActor->GetLevel();
+		if (!UE::Private::ExternalDataLayerHelper::ValidateAssetUsingAssetReferenceRestrictions(Referencer, ActorReferencedAssets, InvalidReferenceReasons))
 		{
 			const FString JoinedReasons = FString::Join(InvalidReferenceReasons, TEXT("\n"));
 			if (NewExternalDataLayerAsset)
@@ -236,7 +217,6 @@ bool FExternalDataLayerHelper::CanMoveActorsToExternalDataLayer(const TArray<AAc
 			}
 			else
 			{
-				const UExternalDataLayerAsset* OldExternalDataLayerAsset = InActor->GetExternalDataLayerAsset();
 				check(OldExternalDataLayerAsset);
 				OutFailureReason = FText::Format(LOCTEXT("CantRemoveEDLFromActorReferenceRestrictions", "Can't remove External Data Layer {0} from Actor {1}. Reason: {2}."), FText::FromString(OldExternalDataLayerAsset->GetName()), FText::FromString(InActor->GetName()), FText::FromString(JoinedReasons));
 			}
diff --git a/Engine/Source/Runtime/Engine/Private/WorldPartition/DataLayer/ExternalDataLayerManager.cpp b/Engine/Source/Runtime/Engine/Private/WorldPartition/DataLayer/ExternalDataLayerManager.cpp
--- a/Engine/Source/Runtime/Engine/Private/WorldPartition/DataLayer/ExternalDataLayerManager.cpp
+++ b/Engine/Source/Runtime/Engine/Private/WorldPartition/DataLayer/ExternalDataLayerManager.cpp
@@ -832,45 +832,6 @@ FString UExternalDataLayerManager::GetActorPackageName(const UExternalDataLayerA
 	return ActorPackageName;
 }
 
-bool UExternalDataLayerManager::SetupActorPackageForExternalDataLayerAsset(AActor* InActor, const UExternalDataLayerAsset* InExternalDataLayerAsset)
-{
-	check(InActor);
-	check(InActor->IsPackageExternal());
-
-	// First check if we really need to rename the package at all.
-	// For example, when reinstancing (after compiling a BP), we reuse the old actor package.
-	const FString OldActorPackageName = InActor->GetPackage()->GetName();
-	const FString NewActorPackageName = GetActorPackageName(InExternalDataLayerAsset, InActor->GetLevel(), InActor->GetPathName());
-	if (OldActorPackageName == NewActorPackageName)
-	{
-		return true;
-	}
-
-	check(InActor->GetExternalPackage()->HasAnyPackageFlags(PKG_NewlyCreated));
-	UExternalDataLayerInstance* ExternalDataLayerInstance = GetExternalDataLayerInstance(InExternalDataLayerAsset);
-	if (!ExternalDataLayerInstance)
-	{
-		UE_LOG(LogWorldPartition, Warning, TEXT("[EDL: %s] Can't find External Data Layer instance, package of actor %s won't be moved under External Data Layer root."), *InExternalDataLayerAsset->GetName(), *InActor->GetActorNameOrLabel());
-		return false;
-	}
-	
-	bool bSuccess = InActor->GetPackage()->Rename(*NewActorPackageName);
-	UE_LOG(LogWorldPartition, Verbose, TEXT("[EDL: %s] Set new actor %s Package %s."), *InExternalDataLayerAsset->GetName(), *InActor->GetActorNameOrLabel(), *InActor->GetPackage()->GetName());
-	FText FailureReason;
-	if (!ExternalDataLayerInstance->CanAddActor(InActor, &FailureReason))
-	{
-		InActor->GetPackage()->Rename(*OldActorPackageName);
-		UE_LOG(LogWorldPartition, Warning, TEXT("[EDL: %s] Can't rename package for actor %s. %s"), *InExternalDataLayerAsset->GetName(), *InActor->GetActorNameOrLabel(), *FailureReason.ToString());
-		return false;
-	}
-	return bSuccess;
-}
-
-bool UExternalDataLayerManager::OnActorPreSpawnInitialization(AActor* InActor, const UExternalDataLayerAsset* InExternalDataLayerAsset)
-{
-	return SetupActorPackageForExternalDataLayerAsset(InActor, InExternalDataLayerAsset);
-}
-
 bool UExternalDataLayerManager::ValidateOnActorExternalDataLayerAssetChanged(AActor* InActor)
 {
 	check(InActor);
diff --git a/Engine/Source/Runtime/Engine/Private/WorldPartition/RuntimeSpatialHash/RuntimeSpatialHashHLOD.cpp b/Engine/Source/Runtime/Engine/Private/WorldPartition/RuntimeSpatialHash/RuntimeSpatialHashHLOD.cpp
--- a/Engine/Source/Runtime/Engine/Private/WorldPartition/RuntimeSpatialHash/RuntimeSpatialHashHLOD.cpp
+++ b/Engine/Source/Runtime/Engine/Private/WorldPartition/RuntimeSpatialHash/RuntimeSpatialHashHLOD.cpp
@@ -17,7 +17,6 @@
 #include "WorldPartition/WorldPartition.h"
 #include "WorldPartition/WorldPartitionHelpers.h"
 #include "WorldPartition/ContentBundle/ContentBundleActivationScope.h"
-#include "ActorEditorContext/ScopedActorEditorContextSetExternalDataLayerAsset.h"
 
 #include "UObject/GCObjectScopeGuard.h"
 #include "UObject/SavePackage.h"
@@ -398,7 +397,7 @@ static void UpdateHLODGridsActors(UWorld* World, const UActorDescContainerInstan
 		if (!GridActor)
 		{
 			FContentBundleActivationScope ContentBndleScope(ContainerInstance->GetContentBundleGuid());
-			FScopedActorEditorContextSetExternalDataLayerAsset EDLScope(ContainerInstance->GetExternalDataLayerAsset());
+			FScopedOverrideSpawningLevelMountPointObject EDLScope(ContainerInstance->GetExternalDataLayerAsset());
 
 			FActorSpawnParameters SpawnParams;
 			SpawnParams.bCreateActorPackage = true;
diff --git a/Engine/Source/Runtime/Engine/Public/WorldPartition/DataLayer/ExternalDataLayerHelper.h b/Engine/Source/Runtime/Engine/Public/WorldPartition/DataLayer/ExternalDataLayerHelper.h
--- a/Engine/Source/Runtime/Engine/Public/WorldPartition/DataLayer/ExternalDataLayerHelper.h
+++ b/Engine/Source/Runtime/Engine/Public/WorldPartition/DataLayer/ExternalDataLayerHelper.h
@@ -68,6 +68,7 @@ private:
 	ENGINE_API static bool MoveActorsToExternalDataLayer(const TArray<AActor*>& InActors, const UExternalDataLayerInstance* InExternalDataLayerInstance, FText* OutFailureReason = nullptr);
 
 	friend class FDataLayerEditorModule;
+	friend class UDataLayerEditorSubsystem;
 #endif
 	static constexpr FStringView GetExternalDataLayerFolder() { return ExternalDataLayerFolder; }
 	static constexpr FStringView ExternalDataLayerFolder = TEXTVIEW("/EDL/");
diff --git a/Engine/Source/Runtime/Engine/Public/WorldPartition/DataLayer/ExternalDataLayerManager.h b/Engine/Source/Runtime/Engine/Public/WorldPartition/DataLayer/ExternalDataLayerManager.h
--- a/Engine/Source/Runtime/Engine/Public/WorldPartition/DataLayer/ExternalDataLayerManager.h
+++ b/Engine/Source/Runtime/Engine/Public/WorldPartition/DataLayer/ExternalDataLayerManager.h
@@ -81,9 +81,7 @@ private:
 	const UExternalDataLayerAsset* GetMatchingExternalDataLayerAssetForObjectPath(const FSoftObjectPath& InObjectPath);
 	const UExternalDataLayerAsset* GetActorEditorContextCurrentExternalDataLayer() const;
 	AWorldDataLayers* GetWorldDataLayers(const UExternalDataLayerAsset* InExternalDataLayerAsset, bool bInAllowCreate = false) const;
-	bool OnActorPreSpawnInitialization(AActor* InActor, const UExternalDataLayerAsset* InExternalDataLayerAsset);
 	FString GetActorPackageName(const UExternalDataLayerAsset* InExternalDataLayerAsset, const ULevel* InDestinationLevel, const FString& InActorPath) const;
-	bool SetupActorPackageForExternalDataLayerAsset(AActor* InActor, const UExternalDataLayerAsset* InExternalDataLayerAsset);
 	URuntimeHashExternalStreamingObjectBase* CreateExternalStreamingObjectUsingStreamingGeneration(const UExternalDataLayerAsset* InExternalDataLayerAsset);
 	bool HasInjectedExternalDataLayerAssets() const { return InjectedExternalDataLayerAssets.Num() > 0; }
 
@@ -127,6 +125,7 @@ private:
 	friend class UExternalDataLayerInstance;
 	friend class UDataLayerEditorSubsystem;
 	friend class ULevelInstanceSubsystem;
+	friend class UContentBundleEditingSubmodule;
 	friend class UWorldPartitionRuntimeLevelStreamingCell;
 	friend class UGameFeatureActionConvertContentBundleWorldPartitionBuilder;
 };
\ No newline at end of file
