Data Layer Outliner : Fixed drop text showing 'create' data layer instead of 'move'  when appropriate.

diff --git a/Engine/Source/Editor/DataLayerEditor/Private/DataLayer/DataLayerMode.cpp b/Engine/Source/Editor/DataLayerEditor/Private/DataLayer/DataLayerMode.cpp
--- a/Engine/Source/Editor/DataLayerEditor/Private/DataLayer/DataLayerMode.cpp
+++ b/Engine/Source/Editor/DataLayerEditor/Private/DataLayer/DataLayerMode.cpp
@@ -702,7 +702,8 @@ FSceneOutlinerDragValidationInfo FDataLayerMode::ValidateDrop(const ISceneOutlin
 			}
 			if (!DataLayerAssets.IsEmpty())
 			{
-				return ValidateDataLayerAssetDrop(DropTarget, DataLayerAssets.Array());
+				const bool bIsMovingDataLayerInstances = true;
+				return ValidateDataLayerAssetDrop(DropTarget, DataLayerAssets.Array(), bIsMovingDataLayerInstances);
 			}
 			return FSceneOutlinerDragValidationInfo(ESceneOutlinerDropCompatibility::IncompatibleGeneric, LOCTEXT("NoValidDataLayersToMove", "No valid Data Layer to move"));
 		}
@@ -793,13 +794,15 @@ bool FDataLayerMode::CanReferenceDataLayerAssets(const AWorldDataLayers* InWorld
 	return true;
 }
 
-FSceneOutlinerDragValidationInfo FDataLayerMode::ValidateDataLayerAssetDrop(const ISceneOutlinerTreeItem& DropTarget, const TArray<const UDataLayerAsset*>& DataLayerAssetsToDrop) const
+FSceneOutlinerDragValidationInfo FDataLayerMode::ValidateDataLayerAssetDrop(const ISceneOutlinerTreeItem& DropTarget, const TArray<const UDataLayerAsset*>& DataLayerAssetsToDrop, bool bIsMovingDataLayerInstances) const
 {
 	check(!DataLayerAssetsToDrop.IsEmpty());
+	const FString CreateOrMoveString(bIsMovingDataLayerInstances ? "move" : "create");
 
 	if (DataLayerEditorSubsystem->HasDeprecatedDataLayers())
 	{
-		return FSceneOutlinerDragValidationInfo(ESceneOutlinerDropCompatibility::Incompatible, FText::Format(LOCTEXT("CantCreateInstanceWorldHasDeprecatedDataLayers", "Cannot create Data Layers Instance from assets since \"{0}\" has deprecated data layers."),
+		return FSceneOutlinerDragValidationInfo(ESceneOutlinerDropCompatibility::Incompatible, FText::Format(LOCTEXT("CantCreateOrMoveInstanceWorldHasDeprecatedDataLayers", "Cannot {0} Data Layers Instance from assets since \"{1}\" has deprecated data layers."),
+			FText::FromString(CreateOrMoveString), 
 			FText::FromString(GetOwningWorld()->GetName())));
 	}
 
@@ -861,7 +864,7 @@ FSceneOutlinerDragValidationInfo FDataLayerMode::ValidateDataLayerAssetDrop(cons
 	FText FailureReason;
 	if (!CanReferenceDataLayerAssets(DropTargetWorldDataLayers, DataLayerAssetsToDrop, &FailureReason))
 	{
-		return FSceneOutlinerDragValidationInfo(ESceneOutlinerDropCompatibility::Incompatible, FText::Format(LOCTEXT("CantCreateDataLayerInstancePassFilterFailed", "Cannot create Data Layer Instance : {0}"), FailureReason));
+		return FSceneOutlinerDragValidationInfo(ESceneOutlinerDropCompatibility::Incompatible, FText::Format(LOCTEXT("CantCreateOrMoveDataLayerInstancePassFilterFailed", "Cannot {0} Data Layer Instance: {1}"), FText::FromString(CreateOrMoveString), FailureReason));
 	}
 
 	// Check if target data layer supports having dropped asset as a child
@@ -873,7 +876,7 @@ FSceneOutlinerDragValidationInfo FDataLayerMode::ValidateDataLayerAssetDrop(cons
 			FText Reason;
 			if (!FDataLayerUtils::AreDataLayerTypesCompatible(ParentType, DataLayerAssetToDrop->GetType(), DropTargetDataLayerWithAsset->IsA<UExternalDataLayerInstance>(), &Reason))
 			{
-				return FSceneOutlinerDragValidationInfo(ESceneOutlinerDropCompatibility::Incompatible, FText::Format(LOCTEXT("CantCreateDataLayerInstanceIncompatibleChildType", "Cannot create Data Layer Instance : {0}"), Reason));
+				return FSceneOutlinerDragValidationInfo(ESceneOutlinerDropCompatibility::Incompatible, FText::Format(LOCTEXT("CantCreateOrMoveDataLayerInstanceIncompatibleChildType", "Cannot {0} Data Layer Instance: {1}"), FText::FromString(CreateOrMoveString), Reason));
 			}
 		}
 	}
diff --git a/Engine/Source/Editor/DataLayerEditor/Private/DataLayer/DataLayerMode.h b/Engine/Source/Editor/DataLayerEditor/Private/DataLayer/DataLayerMode.h
--- a/Engine/Source/Editor/DataLayerEditor/Private/DataLayer/DataLayerMode.h
+++ b/Engine/Source/Editor/DataLayerEditor/Private/DataLayer/DataLayerMode.h
@@ -157,7 +157,7 @@ private:
 	AWorldDataLayers* GetWorldDataLayersFromTreeItem(const ISceneOutlinerTreeItem& TreeItem) const;
 	UDataLayerInstance* GetDataLayerInstanceFromTreeItem(const ISceneOutlinerTreeItem& TreeItem) const;
 	FSceneOutlinerDragValidationInfo ValidateActorDrop(const ISceneOutlinerTreeItem& DropTarget, TArray<AActor*> PayloadActors, bool bMoveOperation = false) const;
-	FSceneOutlinerDragValidationInfo ValidateDataLayerAssetDrop(const ISceneOutlinerTreeItem& DropTarget, const TArray<const UDataLayerAsset*>& DataLayerAssetsToDrop) const;
+	FSceneOutlinerDragValidationInfo ValidateDataLayerAssetDrop(const ISceneOutlinerTreeItem& DropTarget, const TArray<const UDataLayerAsset*>& DataLayerAssetsToDrop, bool bIsMovingDataLayerInstances = false) const;
 	bool CanReferenceDataLayerAssets(const AWorldDataLayers* InWorldDataLayers, const TArray<const UDataLayerAsset*>& InReferencedDataLayerAssets, FText* OutFailureReason = nullptr) const;
 	void OnDataLayerAssetDropped(const TArray<const UDataLayerAsset*>& DroppedDataLayerAsset, ISceneOutlinerTreeItem& DropTarget) const;
 
