Only allow data layers of a game feature plugin to be under a WorldDataLayers of the same game feature plugin.
Only allow data layers of a non-game feature plugin to be under a non-game feature WorldDataLayers.

diff --git a/Engine/Plugins/Editor/DataValidation/Source/DataValidation/Private/WorldPartitionChangelistValidator.cpp b/Engine/Plugins/Editor/DataValidation/Source/DataValidation/Private/WorldPartitionChangelistValidator.cpp
--- a/Engine/Plugins/Editor/DataValidation/Source/DataValidation/Private/WorldPartitionChangelistValidator.cpp
+++ b/Engine/Plugins/Editor/DataValidation/Source/DataValidation/Private/WorldPartitionChangelistValidator.cpp
@@ -126,6 +126,10 @@ void UWorldPartitionChangelistValidator::ValidateActorsAndDataLayersFromChangeLi
 		
 					RelevantDataLayerAssets.Add(AssetData.PackageName.ToString());
 				}
+				else if (AssetClass->IsChildOf<UDataLayerInstance>())
+				{
+					RelevantExternalPackageDataLayerInstances.Add(AssetData.PackageName.ToString());
+				}
 				else if (AssetClass->IsChildOf<UWorld>())
 				{
 					if (ULevel::GetIsLevelPartitionedFromPackage(PackageName))
@@ -273,7 +277,17 @@ bool UWorldPartitionChangelistValidator::Filter(const IWorldPartitionActorDescIn
 bool UWorldPartitionChangelistValidator::Filter(const UDataLayerInstance* InDataLayerInstance)
 {
 	const UDataLayerInstanceWithAsset* DataLayerWithAsset = Cast<UDataLayerInstanceWithAsset>(InDataLayerInstance);
-	return DataLayerWithAsset != nullptr && DataLayerWithAsset->GetAsset() != nullptr && RelevantDataLayerAssets.Contains(DataLayerWithAsset->GetAsset()->GetPathName());
+	if (DataLayerWithAsset && DataLayerWithAsset->GetAsset() && RelevantDataLayerAssets.Contains(DataLayerWithAsset->GetAsset()->GetPackage()->GetName()))
+	{
+		return true;
+	}
+
+	if (InDataLayerInstance->IsPackageExternal() && RelevantExternalPackageDataLayerInstances.Contains(InDataLayerInstance->GetPackage()->GetName()))
+	{
+		return true;
+	}
+
+	return false;
 }
 
 void UWorldPartitionChangelistValidator::OnInvalidRuntimeGrid(const IWorldPartitionActorDescInstanceView& ActorDescView, FName GridName)
@@ -430,6 +444,19 @@ void UWorldPartitionChangelistValidator::OnDataLayerHierarchyTypeMismatch(const
 	}
 }
 
+void UWorldPartitionChangelistValidator::OnInvalidWorldDataLayersReference(const AWorldDataLayers* WorldDataLayers, const UDataLayerInstance* DataLayerInstance, const FText& Reason)
+{
+	if (RelevantActorGuids.Find(WorldDataLayers->GetActorGuid()) || Filter(DataLayerInstance))
+	{
+		FText CurrentError = FText::Format(LOCTEXT("DataValidation.Changelist.WorldPartition.DataLayerAssetReferenceRestrictions", "Actor {0} can't reference data layer {1} because of asset reference restrictions ({2})"),
+			FText::FromString(WorldDataLayers->GetName()),
+			FText::FromString(DataLayerInstance->GetDataLayerFullName()),
+			Reason);
+
+		AssetFails(CurrentAsset, CurrentError);
+	}
+}
+
 void UWorldPartitionChangelistValidator::OnDataLayerAssetConflict(const UDataLayerInstanceWithAsset* DataLayerInstance, const UDataLayerInstanceWithAsset* ConflictingDataLayerInstance)
 {
 	if (Filter(DataLayerInstance)
diff --git a/Engine/Plugins/Editor/DataValidation/Source/DataValidation/Private/WorldPartitionChangelistValidator.h b/Engine/Plugins/Editor/DataValidation/Source/DataValidation/Private/WorldPartitionChangelistValidator.h
--- a/Engine/Plugins/Editor/DataValidation/Source/DataValidation/Private/WorldPartitionChangelistValidator.h
+++ b/Engine/Plugins/Editor/DataValidation/Source/DataValidation/Private/WorldPartitionChangelistValidator.h
@@ -19,6 +19,7 @@ protected:
 	FTopLevelAssetPath RelevantMap;
 	TSet<FGuid> RelevantActorGuids;
 	TSet<FString> RelevantDataLayerAssets;
+	TSet<FString> RelevantExternalPackageDataLayerInstances;
 	UObject* CurrentAsset = nullptr;
 	bool bSubmittingWorldDataLayers = false;
 
@@ -43,6 +44,7 @@ protected:
 	virtual void OnInvalidReferenceDataLayerAsset(const UDataLayerInstanceWithAsset* DataLayerInstance) override;
 	virtual void OnInvalidDataLayerAssetType(const UDataLayerInstanceWithAsset* DataLayerInstance, const UDataLayerAsset* DataLayerAsset) override;
 	virtual void OnDataLayerHierarchyTypeMismatch(const UDataLayerInstance* DataLayerInstance, const UDataLayerInstance* Parent, EDataLayerHierarchyInvalidReason Reason) override;
+	virtual void OnInvalidWorldDataLayersReference(const AWorldDataLayers* WorldDataLayers, const UDataLayerInstance* DataLayerInstance, const FText& Reason) override;
 	virtual void OnDataLayerAssetConflict(const UDataLayerInstanceWithAsset* DataLayerInstance, const UDataLayerInstanceWithAsset* ConflictingDataLayerInstance) override;
 	virtual void OnActorNeedsResave(const IWorldPartitionActorDescInstanceView& ActorDescView) override;
 	virtual void OnLevelInstanceInvalidWorldAsset(const IWorldPartitionActorDescInstanceView& ActorDescView, FName WorldAsset, ELevelInstanceInvalidReason Reason) override;
diff --git a/Engine/Source/Editor/DataLayerEditor/Private/DataLayer/DataLayerInstanceCustomization.cpp b/Engine/Source/Editor/DataLayerEditor/Private/DataLayer/DataLayerInstanceCustomization.cpp
--- a/Engine/Source/Editor/DataLayerEditor/Private/DataLayer/DataLayerInstanceCustomization.cpp
+++ b/Engine/Source/Editor/DataLayerEditor/Private/DataLayer/DataLayerInstanceCustomization.cpp
@@ -9,7 +9,13 @@
 #include "UObject/Object.h"
 #include "UObject/WeakObjectPtr.h"
 #include "UObject/WeakObjectPtrTemplates.h"
+#include "WorldPartition/DataLayer/WorldDataLayers.h"
 #include "WorldPartition/DataLayer/DataLayerInstance.h"
+#include "WorldPartition/DataLayer/DataLayerInstanceWithAsset.h"
+#include "Templates/SharedPointer.h"
+#include "DetailCategoryBuilder.h"
+#include "DetailWidgetRow.h"
+#include "PropertyCustomizationHelpers.h"
 
 #define LOCTEXT_NAMESPACE "FDataLayerInstanceDetails"
 
@@ -23,6 +29,7 @@ void FDataLayerInstanceDetails::CustomizeDetails(IDetailLayoutBuilder& DetailBui
 	TArray<TWeakObjectPtr<UObject>> ObjectsBeingCustomized;
 	DetailBuilder.GetObjectsBeingCustomized(ObjectsBeingCustomized);
 	
+	uint32 CustomizedDataLayerInstanceWithAssetCount = 0;
 	bool bHasRuntime = false;
 	for (const TWeakObjectPtr<UObject>& SelectedObject : ObjectsBeingCustomized)
 	{
@@ -30,13 +37,55 @@ void FDataLayerInstanceDetails::CustomizeDetails(IDetailLayoutBuilder& DetailBui
 		if (DataLayerInstance && DataLayerInstance->IsRuntime() && !DataLayerInstance->IsClientOnly() && !DataLayerInstance->IsServerOnly())
 		{
 			bHasRuntime = true;
-			break;
+		}
+		if (Cast<UDataLayerInstanceWithAsset>(SelectedObject.Get()))
+		{
+			++CustomizedDataLayerInstanceWithAssetCount;
 		}
 	}
 	if (!bHasRuntime)
 	{
 		DetailBuilder.HideProperty(GET_MEMBER_NAME_CHECKED(UDataLayerInstance, InitialRuntimeState));
 	}
+
+	TSharedRef<IPropertyHandle> DataLayerAssetProperty = DetailBuilder.GetProperty("DataLayerAsset", UDataLayerInstanceWithAsset::StaticClass());
+	UDataLayerInstanceWithAsset* CustomizedDataLayerInstanceWithAsset = (ObjectsBeingCustomized.Num() == 1) ? Cast<UDataLayerInstanceWithAsset>(ObjectsBeingCustomized[0].Get()) : nullptr;
+	if (AWorldDataLayers* CustomizedWorldDataLayers = CustomizedDataLayerInstanceWithAsset ? CustomizedDataLayerInstanceWithAsset->GetDirectOuterWorldDataLayers() : nullptr)
+	{
+		IDetailCategoryBuilder& Category = DetailBuilder.EditCategory("Data Layer");
+		Category.AddCustomRow(DataLayerAssetProperty->GetPropertyDisplayName())
+			.RowTag(DataLayerAssetProperty->GetProperty()->GetFName())
+			.NameContent()
+			[
+				DataLayerAssetProperty->CreatePropertyNameWidget()
+			]
+			.ValueContent()
+			.MinDesiredWidth(200.f)
+			[
+				SNew(SObjectPropertyEntryBox)
+				.AllowClear(false)
+				.AllowedClass(UDataLayerAsset::StaticClass())
+				.PropertyHandle(DataLayerAssetProperty)
+				.DisplayThumbnail(true)
+				.ThumbnailPool(DetailBuilder.GetThumbnailPool())
+				.OnShouldFilterAsset_Lambda([CustomizedWorldDataLayers](const FAssetData& AssetData)
+				{
+					FText FailureReason;
+					const UDataLayerAsset* DataLayerAsset = CastChecked<UDataLayerAsset>(AssetData.GetAsset());
+					return !CustomizedWorldDataLayers->CanReferenceDataLayerAsset(DataLayerAsset, &FailureReason);
+				})
+			];
+
+		DetailBuilder.HideProperty(DataLayerAssetProperty);
+	}
+	else if (CustomizedDataLayerInstanceWithAssetCount > 1)
+	{
+		// We don't want to be able to set the same DataLayerAsset on multiple data layer instances
+		if (IDetailPropertyRow* DetailPropertyRow = DetailBuilder.EditDefaultProperty(DataLayerAssetProperty))
+		{
+			DetailPropertyRow->IsEnabled(false);
+		}
+	}
 }
 
 #undef LOCTEXT_NAMESPACE
diff --git a/Engine/Source/Runtime/Engine/Private/WorldPartition/DataLayer/DataLayerInstance.cpp b/Engine/Source/Runtime/Engine/Private/WorldPartition/DataLayer/DataLayerInstance.cpp
--- a/Engine/Source/Runtime/Engine/Private/WorldPartition/DataLayer/DataLayerInstance.cpp
+++ b/Engine/Source/Runtime/Engine/Private/WorldPartition/DataLayer/DataLayerInstance.cpp
@@ -534,6 +534,13 @@ FText UDataLayerInstance::GetDataLayerText(const UDataLayerInstance* InDataLayer
 
 bool UDataLayerInstance::Validate(IStreamingGenerationErrorHandler* ErrorHandler) const
 {
+	// No need to run the validation for game worlds
+	UWorld* World = GetOuterWorld();
+	if (World->IsGameWorld())
+	{
+		return true;
+	}
+
 	if (GetParent() != nullptr)
 	{
 		FText ReasonText;
diff --git a/Engine/Source/Runtime/Engine/Private/WorldPartition/DataLayer/DataLayerInstanceWithAsset.cpp b/Engine/Source/Runtime/Engine/Private/WorldPartition/DataLayer/DataLayerInstanceWithAsset.cpp
--- a/Engine/Source/Runtime/Engine/Private/WorldPartition/DataLayer/DataLayerInstanceWithAsset.cpp
+++ b/Engine/Source/Runtime/Engine/Private/WorldPartition/DataLayer/DataLayerInstanceWithAsset.cpp
@@ -165,6 +165,14 @@ bool UDataLayerInstanceWithAsset::Validate(IStreamingGenerationErrorHandler* Err
 		return false;
 	}
 
+	FText ReasonText;
+	const AWorldDataLayers* WorldDataLayers = GetDirectOuterWorldDataLayers();
+	if (WorldDataLayers && !WorldDataLayers->CanReferenceDataLayerAsset(GetAsset(), &ReasonText))
+	{
+		ErrorHandler->OnInvalidWorldDataLayersReference(WorldDataLayers, this, ReasonText);
+		return false;
+	}
+
 	// Get the DataLayerManager for this DataLayerInstance which will be the one of its outer world
 	UDataLayerManager* DataLayerManager = UDataLayerManager::GetDataLayerManager(this);
 	if (ensure(DataLayerManager))
diff --git a/Engine/Source/Runtime/Engine/Private/WorldPartition/DataLayer/WorldDataLayers.cpp b/Engine/Source/Runtime/Engine/Private/WorldPartition/DataLayer/WorldDataLayers.cpp
--- a/Engine/Source/Runtime/Engine/Private/WorldPartition/DataLayer/WorldDataLayers.cpp
+++ b/Engine/Source/Runtime/Engine/Private/WorldPartition/DataLayer/WorldDataLayers.cpp
@@ -27,6 +27,7 @@
 #include "WorldPartition/WorldPartitionEditorPerProjectUserSettings.h"
 #include "WorldPartition/DataLayer/WorldDataLayersActorDesc.h"
 #include "LevelInstance/LevelInstanceSubsystem.h"
+#include "Interfaces/IPluginManager.h"
 #include "Algo/Find.h"
 #include "Misc/MessageDialog.h"
 #include "Editor.h"
@@ -1369,7 +1370,33 @@ bool AWorldDataLayers::CanReferenceDataLayerAsset(const UDataLayerAsset* InDataL
 
 	const UExternalDataLayerAsset* RootExternalDataLayerAsset = GetRootExternalDataLayerAsset();
 	const UObject* ReferencingObject = RootExternalDataLayerAsset ? Cast<UObject>(RootExternalDataLayerAsset) : Cast<UObject>(this);
-	return PassesAssetReferenceFiltering(ReferencingObject, InDataLayerAsset, OutFailureReason);
+	if (!PassesAssetReferenceFiltering(ReferencingObject, InDataLayerAsset, OutFailureReason))
+	{
+		return false;
+	}
+	else if (InDataLayerAsset)
+	{
+		IPluginManager& PluginManager = IPluginManager::Get();
+		const FString AssetMountPoint = FPackageName::GetPackageMountPoint(InDataLayerAsset->GetPackage()->GetName()).ToString();
+		const FString ThisMountPoint = FPackageName::GetPackageMountPoint(GetPackage()->GetName()).ToString();
+		TSharedPtr<IPlugin> AssetPlugin = PluginManager.FindPluginFromPath(AssetMountPoint);
+		TSharedPtr<IPlugin> ThisPlugin = PluginManager.FindPluginFromPath(ThisMountPoint);
+		if (AssetPlugin != ThisPlugin)
+		{
+			if (OutFailureReason)
+			{
+				*OutFailureReason = FText::Format(LOCTEXT("WorldDataLayerAndDataLayerAssetPluginMismatch", "{0} part of {1}'{2}' cannot be referenced by {3} part of {4}'{5}'."),
+					FText::FromString(InDataLayerAsset->GetName()),
+					FText::FromString(AssetPlugin.IsValid() ? TEXT("plugin ") : TEXT("")),
+					FText::FromString(AssetMountPoint),
+					FText::FromString(GetActorLabel()),
+					FText::FromString(ThisPlugin.IsValid() ? TEXT("plugin ") : TEXT("")),
+					FText::FromString(ThisMountPoint));
+			}
+			return false;
+		}
+	}
+	return true;
 }
 
 void AWorldDataLayers::RemoveEditorDataLayers()
diff --git a/Engine/Source/Runtime/Engine/Private/WorldPartition/ErrorHandling/WorldPartitionStreamingGenerationLogErrorHandler.cpp b/Engine/Source/Runtime/Engine/Private/WorldPartition/ErrorHandling/WorldPartitionStreamingGenerationLogErrorHandler.cpp
--- a/Engine/Source/Runtime/Engine/Private/WorldPartition/ErrorHandling/WorldPartitionStreamingGenerationLogErrorHandler.cpp
+++ b/Engine/Source/Runtime/Engine/Private/WorldPartition/ErrorHandling/WorldPartitionStreamingGenerationLogErrorHandler.cpp
@@ -4,6 +4,7 @@
 #include "WorldPartition/ErrorHandling/WorldPartitionStreamingGenerationLogErrorHandler.h"
 #include "WorldPartition/WorldPartitionLog.h"
 #include "WorldPartition/DataLayer/DataLayerInstanceWithAsset.h"
+#include "WorldPartition/DataLayer/WorldDataLayers.h"
 
 #define UE_ASSET_LOG_ACTORDESCVIEW(CategoryName, Verbosity, ActorDescView, Format, ...) \
 	UE_ASSET_LOG(LogWorldPartition, Log, *ActorDescView.GetActorPackage().ToString(), Format, ##__VA_ARGS__)
@@ -86,6 +87,11 @@ void FStreamingGenerationLogErrorHandler::OnDataLayerHierarchyTypeMismatch(const
 	}
 }
 
+void FStreamingGenerationLogErrorHandler::OnInvalidWorldDataLayersReference(const AWorldDataLayers* WorldDataLayers, const UDataLayerInstance* DataLayerInstance, const FText& Reason)
+{
+	UE_ASSET_LOG(LogWorldPartition, Log, DataLayerInstance, TEXT("Actor %s can't reference data layer %s because of asset reference restrictions (%s)"), *WorldDataLayers->GetName(), *DataLayerInstance->GetDataLayerFullName(), *Reason.ToString());
+}
+
 void FStreamingGenerationLogErrorHandler::OnDataLayerAssetConflict(const UDataLayerInstanceWithAsset* DataLayerInstance, const UDataLayerInstanceWithAsset* ConflictingDataLayerInstance)
 {
 	UE_ASSET_LOG(LogWorldPartition, Log, DataLayerInstance, TEXT("Data Layer Instance %s and Data Layer Instance %s are both referencing Data Layer Asset %s"), *DataLayerInstance->GetDataLayerFName().ToString(), *ConflictingDataLayerInstance->GetDataLayerFName().ToString(), *DataLayerInstance->GetAsset()->GetFullName());
diff --git a/Engine/Source/Runtime/Engine/Private/WorldPartition/ErrorHandling/WorldPartitionStreamingGenerationProxyErrorHandler.cpp b/Engine/Source/Runtime/Engine/Private/WorldPartition/ErrorHandling/WorldPartitionStreamingGenerationProxyErrorHandler.cpp
--- a/Engine/Source/Runtime/Engine/Private/WorldPartition/ErrorHandling/WorldPartitionStreamingGenerationProxyErrorHandler.cpp
+++ b/Engine/Source/Runtime/Engine/Private/WorldPartition/ErrorHandling/WorldPartitionStreamingGenerationProxyErrorHandler.cpp
@@ -48,6 +48,11 @@ void FStreamingGenerationProxyErrorHandler::OnDataLayerHierarchyTypeMismatch(con
 	InnerErrorHandler->OnDataLayerHierarchyTypeMismatch(DataLayerInstance, Parent, Reason);
 }
 
+void FStreamingGenerationProxyErrorHandler::OnInvalidWorldDataLayersReference(const AWorldDataLayers* WorldDataLayers, const UDataLayerInstance* DataLayerInstance, const FText& Reason)
+{
+	InnerErrorHandler->OnInvalidWorldDataLayersReference(WorldDataLayers, DataLayerInstance, Reason);
+}
+
 void FStreamingGenerationProxyErrorHandler::OnDataLayerAssetConflict(const UDataLayerInstanceWithAsset* DataLayerInstance, const UDataLayerInstanceWithAsset* ConflictingDataLayerInstance)
 {
 	InnerErrorHandler->OnDataLayerAssetConflict(DataLayerInstance, ConflictingDataLayerInstance);
diff --git a/Engine/Source/Runtime/Engine/Private/WorldPartition/ErrorHandling/WorldPartitionStreamingGenerationTokenizedMessageErrorHandler.cpp b/Engine/Source/Runtime/Engine/Private/WorldPartition/ErrorHandling/WorldPartitionStreamingGenerationTokenizedMessageErrorHandler.cpp
--- a/Engine/Source/Runtime/Engine/Private/WorldPartition/ErrorHandling/WorldPartitionStreamingGenerationTokenizedMessageErrorHandler.cpp
+++ b/Engine/Source/Runtime/Engine/Private/WorldPartition/ErrorHandling/WorldPartitionStreamingGenerationTokenizedMessageErrorHandler.cpp
@@ -8,6 +8,7 @@
 #include "Misc/MapErrors.h"
 #include "WorldPartition/DataLayer/DataLayerAsset.h"
 #include "WorldPartition/DataLayer/DataLayerInstanceWithAsset.h"
+#include "WorldPartition/DataLayer/WorldDataLayers.h"
 
 #define LOCTEXT_NAMESPACE "WorldPartition"
 
@@ -168,6 +169,20 @@ void ITokenizedMessageErrorHandler::OnDataLayerHierarchyTypeMismatch(const UData
 	HandleTokenizedMessage(MoveTemp(Message));
 }
 
+void ITokenizedMessageErrorHandler::OnInvalidWorldDataLayersReference(const AWorldDataLayers* WorldDataLayers, const UDataLayerInstance* DataLayerInstance, const FText& Reason)
+{
+	TSharedRef<FTokenizedMessage> Message = FTokenizedMessage::Create(EMessageSeverity::Error);
+	Message->AddToken(FTextToken::Create(LOCTEXT("TokenMessage_DataLayers_Actor", "Actor")))
+		->AddToken(FTextToken::Create(FText::FromString(WorldDataLayers->GetName())))
+		->AddToken(FTextToken::Create(LOCTEXT("TokenMessage_DataLayers_CantReferenceDataLayer", "can't reference data layer")))
+		->AddToken(FTextToken::Create(FText::FromString(DataLayerInstance->GetDataLayerShortName())))
+		->AddToken(FTextToken::Create(LOCTEXT("TokenMessage_DataLayers_AssetReferenceRestrictions", "because of asset reference restrictions")))
+		->AddToken(FTextToken::Create(Reason))
+		->AddToken(FMapErrorToken::Create(TEXT("DataLayers_HierarchyTypeMismatch_CheckForErrors")));
+
+	HandleTokenizedMessage(MoveTemp(Message));
+}
+
 void ITokenizedMessageErrorHandler::OnDataLayerAssetConflict(const UDataLayerInstanceWithAsset* DataLayerInstance, const UDataLayerInstanceWithAsset* ConflictingDataLayerInstance)
 {
 	TSharedRef<FTokenizedMessage> Message = FTokenizedMessage::Create(EMessageSeverity::Error);
diff --git a/Engine/Source/Runtime/Engine/Public/WorldPartition/ErrorHandling/WorldPartitionStreamingGenerationErrorHandler.h b/Engine/Source/Runtime/Engine/Public/WorldPartition/ErrorHandling/WorldPartitionStreamingGenerationErrorHandler.h
--- a/Engine/Source/Runtime/Engine/Public/WorldPartition/ErrorHandling/WorldPartitionStreamingGenerationErrorHandler.h
+++ b/Engine/Source/Runtime/Engine/Public/WorldPartition/ErrorHandling/WorldPartitionStreamingGenerationErrorHandler.h
@@ -5,6 +5,7 @@
 #include "CoreMinimal.h"
 #include "WorldPartition/WorldPartitionActorDescInstanceViewInterface.h"
 
+class AWorldDataLayers;
 class UDataLayerAsset;
 class UDataLayerInstance;
 class UDataLayerInstanceWithAsset;
@@ -79,6 +80,11 @@ public:
 	 */
 	virtual void OnDataLayerHierarchyTypeMismatch(const UDataLayerInstance* DataLayerInstance, const UDataLayerInstance* Parent, EDataLayerHierarchyInvalidReason Reason) = 0;
 
+	/**
+ 	 * Called when there's an error with a data layer used by a WorldDataLayers actor
+ 	 */
+	virtual void OnInvalidWorldDataLayersReference(const AWorldDataLayers* WorldDataLayers, const UDataLayerInstance* DataLayerInstance, const FText& Reason) = 0;
+
 	/**
 	 * Called when two data layer instances share the same asset
 	 */
diff --git a/Engine/Source/Runtime/Engine/Public/WorldPartition/ErrorHandling/WorldPartitionStreamingGenerationLogErrorHandler.h b/Engine/Source/Runtime/Engine/Public/WorldPartition/ErrorHandling/WorldPartitionStreamingGenerationLogErrorHandler.h
--- a/Engine/Source/Runtime/Engine/Public/WorldPartition/ErrorHandling/WorldPartitionStreamingGenerationLogErrorHandler.h
+++ b/Engine/Source/Runtime/Engine/Public/WorldPartition/ErrorHandling/WorldPartitionStreamingGenerationLogErrorHandler.h
@@ -16,6 +16,7 @@ public:
 	ENGINE_API virtual void OnInvalidReferenceDataLayerAsset(const UDataLayerInstanceWithAsset* DataLayerInstance) override;
 	ENGINE_API virtual void OnInvalidDataLayerAssetType(const UDataLayerInstanceWithAsset* DataLayerInstance, const UDataLayerAsset* DataLayerAsset) override;
 	ENGINE_API virtual void OnDataLayerHierarchyTypeMismatch(const UDataLayerInstance* DataLayerInstance, const UDataLayerInstance* Parent, EDataLayerHierarchyInvalidReason Reason) override;
+	ENGINE_API virtual void OnInvalidWorldDataLayersReference(const AWorldDataLayers* WorldDataLayers, const UDataLayerInstance* DataLayerInstance, const FText& Reason) override;
 	ENGINE_API virtual void OnDataLayerAssetConflict(const UDataLayerInstanceWithAsset* DataLayerInstance, const UDataLayerInstanceWithAsset* ConflictingDataLayerInstance) override;
 	ENGINE_API virtual void OnActorNeedsResave(const IWorldPartitionActorDescInstanceView& ActorDescView) override;
 	ENGINE_API virtual void OnLevelInstanceInvalidWorldAsset(const IWorldPartitionActorDescInstanceView& ActorDescView, FName WorldAsset, ELevelInstanceInvalidReason Reason) override;
diff --git a/Engine/Source/Runtime/Engine/Public/WorldPartition/ErrorHandling/WorldPartitionStreamingGenerationNullErrorHandler.h b/Engine/Source/Runtime/Engine/Public/WorldPartition/ErrorHandling/WorldPartitionStreamingGenerationNullErrorHandler.h
--- a/Engine/Source/Runtime/Engine/Public/WorldPartition/ErrorHandling/WorldPartitionStreamingGenerationNullErrorHandler.h
+++ b/Engine/Source/Runtime/Engine/Public/WorldPartition/ErrorHandling/WorldPartitionStreamingGenerationNullErrorHandler.h
@@ -16,6 +16,7 @@ public:
 	virtual void OnInvalidReferenceDataLayerAsset(const UDataLayerInstanceWithAsset* DataLayerInstance) override {}
 	virtual void OnInvalidDataLayerAssetType(const UDataLayerInstanceWithAsset* DataLayerInstance, const UDataLayerAsset* DataLayerAsset) override {}
 	virtual void OnDataLayerHierarchyTypeMismatch(const UDataLayerInstance* DataLayerInstance, const UDataLayerInstance* Parent, EDataLayerHierarchyInvalidReason Reason) override {};
+	virtual void OnInvalidWorldDataLayersReference(const AWorldDataLayers* WorldDataLayers, const UDataLayerInstance* DataLayerInstance, const FText& Reason) override {};
 	virtual void OnDataLayerAssetConflict(const UDataLayerInstanceWithAsset* DataLayerInstance, const UDataLayerInstanceWithAsset* ConflictingDataLayerInstance) override {};
 	virtual void OnActorNeedsResave(const IWorldPartitionActorDescInstanceView& ActorDescView) override {};
 	virtual void OnLevelInstanceInvalidWorldAsset(const IWorldPartitionActorDescInstanceView& ActorDescView, FName WorldAsset, ELevelInstanceInvalidReason Reason) override {};
diff --git a/Engine/Source/Runtime/Engine/Public/WorldPartition/ErrorHandling/WorldPartitionStreamingGenerationProxyErrorHandler.h b/Engine/Source/Runtime/Engine/Public/WorldPartition/ErrorHandling/WorldPartitionStreamingGenerationProxyErrorHandler.h
--- a/Engine/Source/Runtime/Engine/Public/WorldPartition/ErrorHandling/WorldPartitionStreamingGenerationProxyErrorHandler.h
+++ b/Engine/Source/Runtime/Engine/Public/WorldPartition/ErrorHandling/WorldPartitionStreamingGenerationProxyErrorHandler.h
@@ -20,6 +20,7 @@ public:
 	ENGINE_API virtual void OnInvalidReferenceDataLayerAsset(const UDataLayerInstanceWithAsset* DataLayerInstance) override;
 	ENGINE_API virtual void OnInvalidDataLayerAssetType(const UDataLayerInstanceWithAsset* DataLayerInstance, const UDataLayerAsset* DataLayerAsset) override;
 	ENGINE_API virtual void OnDataLayerHierarchyTypeMismatch(const UDataLayerInstance* DataLayerInstance, const UDataLayerInstance* Parent, EDataLayerHierarchyInvalidReason Reason) override;
+	ENGINE_API virtual void OnInvalidWorldDataLayersReference(const AWorldDataLayers* WorldDataLayers, const UDataLayerInstance* DataLayerInstance, const FText& Reason) override;
 	ENGINE_API virtual void OnDataLayerAssetConflict(const UDataLayerInstanceWithAsset* DataLayerInstance, const UDataLayerInstanceWithAsset* ConflictingDataLayerInstance) override;
 	ENGINE_API virtual void OnActorNeedsResave(const IWorldPartitionActorDescInstanceView& ActorDescView) override;
 	ENGINE_API virtual void OnLevelInstanceInvalidWorldAsset(const IWorldPartitionActorDescInstanceView& ActorDescView, FName WorldAsset, ELevelInstanceInvalidReason Reason) override;
diff --git a/Engine/Source/Runtime/Engine/Public/WorldPartition/ErrorHandling/WorldPartitionStreamingGenerationTokenizedMessageErrorHandler.h b/Engine/Source/Runtime/Engine/Public/WorldPartition/ErrorHandling/WorldPartitionStreamingGenerationTokenizedMessageErrorHandler.h
--- a/Engine/Source/Runtime/Engine/Public/WorldPartition/ErrorHandling/WorldPartitionStreamingGenerationTokenizedMessageErrorHandler.h
+++ b/Engine/Source/Runtime/Engine/Public/WorldPartition/ErrorHandling/WorldPartitionStreamingGenerationTokenizedMessageErrorHandler.h
@@ -17,6 +17,7 @@ class ITokenizedMessageErrorHandler : public IStreamingGenerationErrorHandler
 	ENGINE_API virtual void OnInvalidReferenceDataLayerAsset(const UDataLayerInstanceWithAsset* DataLayerInstance) override;
 	ENGINE_API virtual void OnInvalidDataLayerAssetType(const UDataLayerInstanceWithAsset* DataLayerInstance, const UDataLayerAsset* DataLayerAsset) override;
 	ENGINE_API virtual void OnDataLayerHierarchyTypeMismatch(const UDataLayerInstance* DataLayerInstance, const UDataLayerInstance* Parent, EDataLayerHierarchyInvalidReason Reason) override;
+	ENGINE_API virtual void OnInvalidWorldDataLayersReference(const AWorldDataLayers* WorldDataLayers, const UDataLayerInstance* DataLayerInstance, const FText& Reason) override;
 	ENGINE_API virtual void OnDataLayerAssetConflict(const UDataLayerInstanceWithAsset* DataLayerInstance, const UDataLayerInstanceWithAsset* ConflictingDataLayerInstance) override;
 	ENGINE_API virtual void OnActorNeedsResave(const IWorldPartitionActorDescInstanceView& ActorDescView) override;
 	ENGINE_API virtual void OnLevelInstanceInvalidWorldAsset(const IWorldPartitionActorDescInstanceView& ActorDescView, FName WorldAsset, ELevelInstanceInvalidReason Reason) override;
