- External Data Layers :  Fixes the case where a child DL of an EDL is made current (which makes its EDL current), didn't remove the Current CB
- AWorldDataLayers::AddToActorEditorContext is now responsible of deactivating Content Bundle Context (instead of UDataLayerEditorSubsystem::AddToActorEditorContext)

diff --git a/Engine/Source/Editor/DataLayerEditor/Private/DataLayer/DataLayerEditorSubsystem.cpp b/Engine/Source/Editor/DataLayerEditor/Private/DataLayer/DataLayerEditorSubsystem.cpp
--- a/Engine/Source/Editor/DataLayerEditor/Private/DataLayer/DataLayerEditorSubsystem.cpp
+++ b/Engine/Source/Editor/DataLayerEditor/Private/DataLayer/DataLayerEditorSubsystem.cpp
@@ -65,7 +65,6 @@
 #include "WorldPartition/DataLayer/ExternalDataLayerEngineSubsystem.h"
 #include "WorldPartition/DataLayer/WorldDataLayers.h"
 #include "WorldPartition/WorldPartition.h"
-#include "WorldPartition/ContentBundle/ContentBundleEditorSubsystemInterface.h"
 
 class SWidget;
 
@@ -577,11 +576,6 @@ TSharedRef<SWidget> UDataLayerEditorSubsystem::GetActorEditorContextWidget(UWorl
 
 void UDataLayerEditorSubsystem::AddToActorEditorContext(UDataLayerInstance* InDataLayerInstance)
 {
-	// EDL has priority over Content Bundle
-	if (InDataLayerInstance->IsA<UExternalDataLayerInstance>())
-	{
-		IContentBundleEditorSubsystemInterface::Get()->DeactivateCurrentContentBundleEditing();
-	}
 	check(InDataLayerInstance->CanBeInActorEditorContext());
 	if (InDataLayerInstance->AddToActorEditorContext())
 	{
diff --git a/Engine/Source/Runtime/Engine/Private/WorldPartition/DataLayer/WorldDataLayers.cpp b/Engine/Source/Runtime/Engine/Private/WorldPartition/DataLayer/WorldDataLayers.cpp
--- a/Engine/Source/Runtime/Engine/Private/WorldPartition/DataLayer/WorldDataLayers.cpp
+++ b/Engine/Source/Runtime/Engine/Private/WorldPartition/DataLayer/WorldDataLayers.cpp
@@ -26,6 +26,7 @@
 #if WITH_EDITOR
 #include "WorldPartition/WorldPartitionEditorPerProjectUserSettings.h"
 #include "WorldPartition/DataLayer/WorldDataLayersActorDesc.h"
+#include "WorldPartition/ContentBundle/ContentBundleEditorSubsystemInterface.h"
 #include "LevelInstance/LevelInstanceSubsystem.h"
 #include "Interfaces/IPluginManager.h"
 #include "Algo/Find.h"
@@ -891,6 +892,9 @@ bool AWorldDataLayers::AddToActorEditorContext(UDataLayerInstance* InDataLayerIn
 			CurrentDataLayers.DataLayerInstanceNames.Remove(DataLayerInstanceName);
 		}
 
+		// EDL has priority over Content Bundle
+		IContentBundleEditorSubsystemInterface::Get()->DeactivateCurrentContentBundleEditing();
+
 		bSuccess = true;
 	}
 	else if (!CurrentDataLayers.DataLayerInstanceNames.Contains(InDataLayerInstance->GetDataLayerFName()))
