- Added External Data Layer support for "Replace Actor with asset/BP class".
- Added External Data Layer support for "Place Actor with asset/BP class".
- Fixed External Data Layer support for BP reinstancing
- When no actor editor context is provided, try to solve a matching EDL based on asset/BP Class mountpoint
- Actor Editor Context EDL is applied when possible, when not, it generates a warning with the details of the failure reason (typically invalid reference)

diff --git a/Engine/Source/Editor/DataLayerEditor/Private/DataLayer/DataLayerEditorSubsystem.cpp b/Engine/Source/Editor/DataLayerEditor/Private/DataLayer/DataLayerEditorSubsystem.cpp
--- a/Engine/Source/Editor/DataLayerEditor/Private/DataLayer/DataLayerEditorSubsystem.cpp
+++ b/Engine/Source/Editor/DataLayerEditor/Private/DataLayer/DataLayerEditorSubsystem.cpp
@@ -65,6 +65,7 @@
 #include "WorldPartition/DataLayer/ExternalDataLayerEngineSubsystem.h"
 #include "WorldPartition/DataLayer/WorldDataLayers.h"
 #include "WorldPartition/WorldPartition.h"
+#include "WorldPartition/IWorldPartitionEditorModule.h"
 
 class SWidget;
 
@@ -96,7 +97,10 @@ private:
 	void Initialize();
 	void OnEditorMapChange(uint32 MapChangeFlags = 0) { DataLayerEditorSubsystem->EditorMapChange(); }
 	void OnPostUndoRedo() { DataLayerEditorSubsystem->PostUndoRedo(); }
+	void OnNewActorsPlaced(UObject* ObjToUse, const TArray<AActor*>& PlacedActors) { DataLayerEditorSubsystem->OnNewActorsPlaced(ObjToUse, PlacedActors); }
+	void OnEditorActorReplaced(AActor* OldActor, AActor* NewActor) {  DataLayerEditorSubsystem->OnEditorActorReplaced(OldActor, NewActor); }
 	void OnCurrentLevelChanged(ULevel* InNewLevel, ULevel* InOldLevel, UWorld* InWorld) { DataLayerEditorSubsystem->EditorRefreshDataLayerBrowser(); }
+	void OnPostWorldInitialization(UWorld* World, const UWorld::InitializationValues IVS) { if (World == DataLayerEditorSubsystem->GetWorld()) { DataLayerEditorSubsystem->EditorMapChange(); } }
 	void OnObjectPostEditChange(UObject* Object, FPropertyChangedEvent& PropertyChangedEvent);
 	void OnLevelActorsAdded(AActor* InActor) { DataLayerEditorSubsystem->InitializeNewActorDataLayers(InActor); }
 	void OnLevelSelectionChanged(UObject* InObject) { DataLayerEditorSubsystem->OnSelectionChanged(); }
@@ -129,7 +133,10 @@ void FDataLayersBroadcast::Deinitialize()
 		{
 			FEditorDelegates::MapChange.RemoveAll(this);
 			FEditorDelegates::PostUndoRedo.RemoveAll(this);
+			FEditorDelegates::OnNewActorsPlaced.RemoveAll(this);
+			FEditorDelegates::OnEditorActorReplaced.RemoveAll(this);
 			FWorldDelegates::OnCurrentLevelChanged.RemoveAll(this);
+			FWorldDelegates::OnPostWorldInitialization.RemoveAll(this);
 			FCoreUObjectDelegates::OnObjectPropertyChanged.RemoveAll(this);
 			if (GEngine)
 			{
@@ -154,7 +161,10 @@ void FDataLayersBroadcast::Initialize()
 		bIsInitialized = true;
 		FEditorDelegates::MapChange.AddRaw(this, &FDataLayersBroadcast::OnEditorMapChange);
 		FEditorDelegates::PostUndoRedo.AddRaw(this, &FDataLayersBroadcast::OnPostUndoRedo);
+		FEditorDelegates::OnNewActorsPlaced.AddRaw(this, &FDataLayersBroadcast::OnNewActorsPlaced);
+		FEditorDelegates::OnEditorActorReplaced.AddRaw(this, &FDataLayersBroadcast::OnEditorActorReplaced);
 		FWorldDelegates::OnCurrentLevelChanged.AddRaw(this, &FDataLayersBroadcast::OnCurrentLevelChanged);
+		FWorldDelegates::OnPostWorldInitialization.AddRaw(this, &FDataLayersBroadcast::OnPostWorldInitialization);
 		FCoreUObjectDelegates::OnObjectPropertyChanged.AddRaw(this, &FDataLayersBroadcast::OnObjectPostEditChange);
 		GEngine->OnLevelActorAdded().AddRaw(this, &FDataLayersBroadcast::OnLevelActorsAdded);
 		USelection::SelectionChangedEvent.AddRaw(this, &FDataLayersBroadcast::OnLevelSelectionChanged);
@@ -367,34 +377,81 @@ void UDataLayerEditorSubsystem::Deinitialize()
 #endif
 }
 
-void UDataLayerEditorSubsystem::OnActorPreSpawnInitialization(AActor* InActor)
+const UExternalDataLayerInstance* UDataLayerEditorSubsystem::GetActorSpawningExternalDataLayerInstance(AActor* InActor) const
 {
-	check(!InActor->GetExternalDataLayerAsset());
-	if (InActor->bIsEditorPreviewActor || !InActor->IsPackageExternal())
+	if (GIsReinstancing || InActor->bIsEditorPreviewActor || !InActor->IsPackageExternal())
 	{
-		return;
+		return nullptr;
 	}
-
+		
 	const UWorld* OwningWorld = InActor->GetWorld();
-	if (OwningWorld && (OwningWorld == GetWorld()))
+	if (!OwningWorld || (OwningWorld != GetWorld()))
+	{
+		return nullptr;
+	}
+
+	// For backward compatibility, don't resolve an External Data Layer when there's a Content Bundle in the Actor Editor Context
+	if (IWorldPartitionEditorModule::Get().IsEditingContentBundle())
+	{
+		return nullptr;
+	}
+
+	const UObject* OverrideSpawningLevelMountPointObject = ULevel::GetOverrideSpawningLevelMountPointObject();
+	const UExternalDataLayerManager* ExternalDataLayerManager = UExternalDataLayerManager::GetExternalDataLayerManager(OwningWorld);
+	
+	// Try to get the external data layer for from the override spawning object
+	const UExternalDataLayerAsset* ResolvedExternalDataLayerAsset = FExternalDataLayerHelper::GetExternalDataLayerAssetFromObject(OverrideSpawningLevelMountPointObject);
+	if (!ResolvedExternalDataLayerAsset)
+	{
+		// If none found, try matching an external data layer with the override spawning object
+		ResolvedExternalDataLayerAsset = ExternalDataLayerManager ? ExternalDataLayerManager->GetMatchingExternalDataLayerAssetForObjectPath(OverrideSpawningLevelMountPointObject) : nullptr;
+	}
+	if (!ResolvedExternalDataLayerAsset)
+	{
+		// If none found, try matching an external data layer with the actor class
+		ResolvedExternalDataLayerAsset = ExternalDataLayerManager ? ExternalDataLayerManager->GetMatchingExternalDataLayerAssetForObjectPath(InActor->GetClass()) : nullptr;
+	}
+	if (!ResolvedExternalDataLayerAsset)
 	{
-		// Prefer override spawning External Data Layer over actor editor context External Data Layer
-		const UExternalDataLayerAsset* OverrideSpawningExternalDataLayer = Cast<UExternalDataLayerAsset>(ULevel::GetOverrideSpawningLevelMountPointObject());
-		const UExternalDataLayerAsset* ExternalDataLayerAsset = OverrideSpawningExternalDataLayer ? OverrideSpawningExternalDataLayer : GetActorEditorContextCurrentExternalDataLayer();
+		// Fallback on actor editor context external data layer
+		ResolvedExternalDataLayerAsset = GetActorEditorContextCurrentExternalDataLayer();
+	}
 
-		UExternalDataLayerManager* ExternalDataLayerManager = UExternalDataLayerManager::GetExternalDataLayerManager(OwningWorld);
-		if (const UExternalDataLayerInstance* ExternalDataLayerInstance = ExternalDataLayerAsset ? ExternalDataLayerManager->GetExternalDataLayerInstance(ExternalDataLayerAsset) : nullptr)
+	const UExternalDataLayerInstance* ResolvedExternalDataLayerInstance = ResolvedExternalDataLayerAsset ? ExternalDataLayerManager->GetExternalDataLayerInstance(ResolvedExternalDataLayerAsset) : nullptr;
+	return ResolvedExternalDataLayerInstance;
+}
+
+void UDataLayerEditorSubsystem::OnActorPreSpawnInitialization(AActor* InActor)
+{
+	if (const UExternalDataLayerInstance* ExternalDataLayerInstance = GetActorSpawningExternalDataLayerInstance(InActor))
+	{
+		check(!InActor->GetExternalDataLayerAsset());
+		FText FailureReason;
+		if (!FExternalDataLayerHelper::MoveActorsToExternalDataLayer({ InActor }, ExternalDataLayerInstance, &FailureReason))
 		{
-			FText FailureReason;
-			if (!FExternalDataLayerHelper::MoveActorsToExternalDataLayer({ InActor }, ExternalDataLayerInstance, &FailureReason))
-			{
-				UE_LOG(LogWorldPartition, Warning, TEXT("%s"), *FailureReason.ToString());
-				LastWarningNotification = FailureReason;
-			}
+			UE_LOG(LogWorldPartition, Warning, TEXT("%s"), *FailureReason.ToString());
+			LastWarningNotification = FailureReason;
 		}
 	}
 }
 
+void UDataLayerEditorSubsystem::OnEditorActorReplaced(AActor* InOldActor, AActor* InNewActor)
+{
+	// Try to apply the current context on the new replacing actor
+	const bool bForceTryApply = true;
+	ApplyContext(InNewActor, bForceTryApply);
+}
+
+void UDataLayerEditorSubsystem::OnNewActorsPlaced(UObject* InObjToUse, const TArray<AActor*>& InPlacedActors)
+{
+	for (AActor* PlacedActor : InPlacedActors)
+	{
+		// Try to apply the current context after actor is placed
+		const bool bForceTryApply = true;
+		ApplyContext(PlacedActor, bForceTryApply);
+	}
+}
+
 UWorld* UDataLayerEditorSubsystem::GetTickableGameObjectWorld() const
 {
 	return GetWorld();
@@ -453,6 +510,46 @@ void UDataLayerEditorSubsystem::BeginDestroy()
 	Super::BeginDestroy();
 }
 
+void UDataLayerEditorSubsystem::ApplyContext(AActor* InActor, bool bInForceTryApply)
+{
+	if (GIsReinstancing || InActor->bIsEditorPreviewActor || !InActor->IsPackageExternal() || (InActor->GetWorld() != GetWorld()))
+	{
+		return;
+	}
+
+	// Try to apply context External Data Layer (this operation can fail if asset referencing validation fails)
+	const UExternalDataLayerAsset* CurrentExternalDataLayer = GetActorEditorContextCurrentExternalDataLayer();
+	const UExternalDataLayerAsset* ActorExternalDataLayerAsset = InActor->GetExternalDataLayerAsset();
+	if (CurrentExternalDataLayer && (ActorExternalDataLayerAsset != CurrentExternalDataLayer))
+	{
+		// Don't apply if there's a valid override spawning External Data Layer (see OnActorPreSpawnInitialization)
+		const UExternalDataLayerInstance* SpawningExternalDataLayer = ULevel::GetOverrideSpawningLevelMountPointObject() ? GetActorSpawningExternalDataLayerInstance(InActor) : nullptr;
+		if (!SpawningExternalDataLayer || bInForceTryApply)
+		{
+			UExternalDataLayerManager* ExternalDataLayerManager = UExternalDataLayerManager::GetExternalDataLayerManager(InActor);
+			if (const UExternalDataLayerInstance* CurrentExternalDataLayerInstance = CurrentExternalDataLayer ? ExternalDataLayerManager->GetExternalDataLayerInstance(CurrentExternalDataLayer) : nullptr)
+			{
+				FText FailureReason;
+				if (!FExternalDataLayerHelper::MoveActorsToExternalDataLayer({ InActor }, CurrentExternalDataLayerInstance, &FailureReason))
+				{
+					UE_LOG(LogWorldPartition, Warning, TEXT("%s"), *FailureReason.ToString());
+					LastWarningNotification = FailureReason;
+				}
+			}
+		}
+	}
+
+	// Apply context Data Layers (except External Data Layer)
+	UDataLayerManager* DataLayerManager = UDataLayerManager::GetDataLayerManager(InActor->GetWorld());
+	TArray<UDataLayerInstance*> DataLayerInstances = DataLayerManager->GetActorEditorContextDataLayers();
+	if (!DataLayerInstances.IsEmpty())
+	{
+		DataLayerInstances.SetNum(Algo::RemoveIf(DataLayerInstances, [](UDataLayerInstance* DataLayerInstance) { return DataLayerInstance->IsA<UExternalDataLayerInstance>(); }));
+		AddActorToDataLayers(InActor, DataLayerInstances);
+		InActor->FixupDataLayers();
+	}
+}
+
 void UDataLayerEditorSubsystem::OnExecuteActorEditorContextAction(UWorld* InWorld, const EActorEditorContextAction& InType, AActor* InActor)
 {
 	UE_CLOG(!InWorld, LogDataLayerEditorSubsystem, Error, TEXT("%s - Failed because world in null."), ANSI_TO_TCHAR(__FUNCTION__));
@@ -463,36 +560,7 @@ void UDataLayerEditorSubsystem::OnExecuteActorEditorContextAction(UWorld* InWorl
 		case EActorEditorContextAction::ApplyContext:
 			check(InActor && InActor->GetWorld() == InWorld);
 			{
-				// Try to apply context External Data Layer (this operation can fail if asset referencing validation fails)
-				// Don't apply if there's a valid override spawning External Data Layer (see OnActorPreSpawnInitialization)
-				const UExternalDataLayerAsset* OverrideSpawningExternalDataLayer = Cast<UExternalDataLayerAsset>(ULevel::GetOverrideSpawningLevelMountPointObject());
-				if (!OverrideSpawningExternalDataLayer)
-				{
-					const UExternalDataLayerAsset* CurrentExternalDataLayer = GetActorEditorContextCurrentExternalDataLayer();
-					const UExternalDataLayerAsset* ActorExternalDataLayerAsset = InActor->GetExternalDataLayerAsset();
-					if (CurrentExternalDataLayer && (ActorExternalDataLayerAsset != CurrentExternalDataLayer))
-					{
-						UExternalDataLayerManager* ExternalDataLayerManager = UExternalDataLayerManager::GetExternalDataLayerManager(InActor);
-						if (const UExternalDataLayerInstance* CurrentExternalDataLayerInstance = CurrentExternalDataLayer ? ExternalDataLayerManager->GetExternalDataLayerInstance(CurrentExternalDataLayer) : nullptr)
-						{
-							FText FailureReason;
-							if (!FExternalDataLayerHelper::MoveActorsToExternalDataLayer({ InActor }, CurrentExternalDataLayerInstance, &FailureReason))
-							{
-								UE_LOG(LogWorldPartition, Warning, TEXT("%s"), *FailureReason.ToString());
-								LastWarningNotification = FailureReason;
-							}
-						}
-					}
-				}
-
-				// Apply context Data Layers (except External Data Layer)
-				TArray<UDataLayerInstance*> DataLayerInstances = DataLayerManager->GetActorEditorContextDataLayers();
-				if (!DataLayerInstances.IsEmpty())
-				{
-					DataLayerInstances.SetNum(Algo::RemoveIf(DataLayerInstances, [](UDataLayerInstance* DataLayerInstance) { return DataLayerInstance->IsA<UExternalDataLayerInstance>(); }));
-					AddActorToDataLayers(InActor, DataLayerInstances);
-					InActor->FixupDataLayers();
-				}
+				ApplyContext(InActor);
 			}
 			break;
 		case EActorEditorContextAction::ResetContext:
@@ -1731,7 +1799,7 @@ void UDataLayerEditorSubsystem::OnSelectionChanged()
 
 const UExternalDataLayerAsset* UDataLayerEditorSubsystem::GetReferencingWorldSurrogateObjectForObject(UWorld* ReferencingWorld, const FSoftObjectPath& ObjectPath)
 {
-	UExternalDataLayerManager* ExternalDataLayerManager = UExternalDataLayerManager::GetExternalDataLayerManager(ReferencingWorld);
+	const UExternalDataLayerManager* ExternalDataLayerManager = UExternalDataLayerManager::GetExternalDataLayerManager(ReferencingWorld);
 	if (const UExternalDataLayerAsset* ExternalDataLayerAsset = ExternalDataLayerManager ? ExternalDataLayerManager->GetMatchingExternalDataLayerAssetForObjectPath(ObjectPath) : nullptr)
 	{
 		const UDataLayerInstance* ExternalDataLayerInstance = UDataLayerManager::GetDataLayerManager(ReferencingWorld)->GetDataLayerInstance(ExternalDataLayerAsset);
@@ -1782,9 +1850,13 @@ private:
 
 TUniquePtr<FLevelEditorDragDropWorldSurrogateReferencingObject> UDataLayerEditorSubsystem::OnLevelEditorDragDropWorldSurrogateReferencingObject(UWorld* ReferencingWorld, const FSoftObjectPath& Object)
 {
-	if (const UExternalDataLayerAsset* ExternalDataLayerAsset = GetReferencingWorldSurrogateObjectForObject(ReferencingWorld, Object))
+	// For backward compatibility, don't try to find a world surrogate object when there's a Content Bundle in the actor editor context
+	if (!IWorldPartitionEditorModule::Get().IsEditingContentBundle())
 	{
-		return MakeUnique<FExternalDataLayerWorldSurrogateReferencingObject>(ExternalDataLayerAsset);
+		if (const UExternalDataLayerAsset* ExternalDataLayerAsset = GetReferencingWorldSurrogateObjectForObject(ReferencingWorld, Object))
+		{
+			return MakeUnique<FExternalDataLayerWorldSurrogateReferencingObject>(ExternalDataLayerAsset);
+		}
 	}
 	return nullptr;
 }
diff --git a/Engine/Source/Editor/DataLayerEditor/Public/DataLayer/DataLayerEditorSubsystem.h b/Engine/Source/Editor/DataLayerEditor/Public/DataLayer/DataLayerEditorSubsystem.h
--- a/Engine/Source/Editor/DataLayerEditor/Public/DataLayer/DataLayerEditorSubsystem.h
+++ b/Engine/Source/Editor/DataLayerEditor/Public/DataLayer/DataLayerEditorSubsystem.h
@@ -27,6 +27,7 @@
 #include "WorldPartition/DataLayer/ActorDataLayer.h"
 #include "WorldPartition/DataLayer/DataLayerInstance.h"
 #include "WorldPartition/DataLayer/WorldDataLayers.h"
+#include "LevelEditorDragDropHandler.h"
 
 #include "DataLayerEditorSubsystem.generated.h"
 
@@ -44,10 +45,10 @@ class UObject;
 class UWorld;
 class UWorldPartition;
 class UExternalDataLayerAsset;
-struct FLevelEditorDragDropWorldSurrogateReferencingObject;
 struct FFrame;
 enum class EExternalDataLayerRegistrationState : uint8;
 template<typename TItemType> class IFilter;
+namespace WorldPartitionTests { class FExternalDataLayerTest; }
 
 USTRUCT(BlueprintType)
 struct DATALAYEREDITOR_API FDataLayerCreationParameters
@@ -760,10 +761,14 @@ private:
 	bool IsActorValidForDataLayerForClasses(AActor* Actor, const TSet<TSubclassOf<UDataLayerInstance>>& DataLayerInstanceClasses);
 
 	// External Data Layer methods
-	void OnActorPreSpawnInitialization(AActor* Actor);
-	static const UExternalDataLayerAsset* GetReferencingWorldSurrogateObjectForObject(UWorld* ReferencingWorld, const FSoftObjectPath& ObjectPath);
-	TUniquePtr<FLevelEditorDragDropWorldSurrogateReferencingObject> OnLevelEditorDragDropWorldSurrogateReferencingObject(UWorld* ReferencingWorld, const FSoftObjectPath& Object);
+	void OnActorPreSpawnInitialization(AActor* InActor);
+	void OnNewActorsPlaced(UObject* InObjToUse, const TArray<AActor*>& InPlacedActors);
+	void OnEditorActorReplaced(AActor* InOldActor, AActor* InNewActor);
 	void OnExternalDataLayerAssetRegistrationStateChanged(const UExternalDataLayerAsset* ExternalDataLayerAsset, EExternalDataLayerRegistrationState OldState, EExternalDataLayerRegistrationState NewState);
+	TUniquePtr<FLevelEditorDragDropWorldSurrogateReferencingObject> OnLevelEditorDragDropWorldSurrogateReferencingObject(UWorld* ReferencingWorld, const FSoftObjectPath& Object);
+	void ApplyContext(AActor* InActor, bool bInForceTryApply = false);
+	const UExternalDataLayerInstance* GetActorSpawningExternalDataLayerInstance(AActor* InActor) const;
+	static const UExternalDataLayerAsset* GetReferencingWorldSurrogateObjectForObject(UWorld* ReferencingWorld, const FSoftObjectPath& ObjectPath);
 
 	/** Contains Data Layers that contain actors that are part of the editor selection */
 	mutable TSet<TWeakObjectPtr<const UDataLayerInstance>> SelectedDataLayersFromEditorSelection;
@@ -807,7 +812,9 @@ private:
 	TOptional<FText> LastWarningNotification;
 
 	friend class FDataLayersBroadcast;
+	friend class UContentBundleEditingSubmodule;
 	friend struct FExternalDataLayerWorldSurrogateReferencingObject;
+	friend class WorldPartitionTests::FExternalDataLayerTest;
 };
 
 template<class DataLayerInstanceType, typename ...Args>
diff --git a/Engine/Source/Editor/UnrealEd/Private/Editor.cpp b/Engine/Source/Editor/UnrealEd/Private/Editor.cpp
--- a/Engine/Source/Editor/UnrealEd/Private/Editor.cpp
+++ b/Engine/Source/Editor/UnrealEd/Private/Editor.cpp
@@ -138,6 +138,7 @@ FEditorDelegates::FOnAssetPostImport					FEditorDelegates::OnAssetPostImport;
 FEditorDelegates::FOnAssetReimport						FEditorDelegates::OnAssetReimport;
 FEditorDelegates::FOnNewActorsDropped					FEditorDelegates::OnNewActorsDropped;
 FEditorDelegates::FOnNewActorsPlaced					FEditorDelegates::OnNewActorsPlaced;
+FEditorDelegates::FOnEditorActorReplaced				FEditorDelegates::OnEditorActorReplaced;
 FEditorDelegates::FOnGridSnappingChanged				FEditorDelegates::OnGridSnappingChanged;
 FSimpleMulticastDelegate								FEditorDelegates::OnLightingBuildStarted;
 FSimpleMulticastDelegate								FEditorDelegates::OnLightingBuildKept;
diff --git a/Engine/Source/Editor/UnrealEd/Private/Factories/ActorFactory.cpp b/Engine/Source/Editor/UnrealEd/Private/Factories/ActorFactory.cpp
--- a/Engine/Source/Editor/UnrealEd/Private/Factories/ActorFactory.cpp
+++ b/Engine/Source/Editor/UnrealEd/Private/Factories/ActorFactory.cpp
@@ -464,6 +464,13 @@ bool UActorFactory::PreSpawnActor( UObject* Asset, FTransform& InOutLocation)
 
 AActor* UActorFactory::SpawnActor(UObject* InAsset, ULevel* InLevel, const FTransform& InTransform, const FActorSpawnParameters& InSpawnParams)
 {
+	// If not already set, provide an override spawning level mount point object using the asset used to spawn the actor
+	TUniquePtr<FScopedOverrideSpawningLevelMountPointObject> ScopeSpawningMountPoint;
+	if (!ULevel::GetOverrideSpawningLevelMountPointObject())
+	{
+		ScopeSpawningMountPoint = MakeUnique<FScopedOverrideSpawningLevelMountPointObject>(InAsset);
+	}
+	
 	ULevel* LocalLevel = ValidateSpawnActorLevel(InLevel, InSpawnParams);
 
 	AActor* DefaultActor = GetDefaultActor(FAssetData(InAsset));
diff --git a/Engine/Source/Editor/UnrealEd/Public/Editor.h b/Engine/Source/Editor/UnrealEd/Public/Editor.h
--- a/Engine/Source/Editor/UnrealEd/Public/Editor.h
+++ b/Engine/Source/Editor/UnrealEd/Public/Editor.h
@@ -114,6 +114,8 @@ struct FEditorDelegates
 	DECLARE_MULTICAST_DELEGATE_TwoParams(FOnNewActorsDropped, const TArray<UObject*>&, const TArray<AActor*>&);
 	/** delegate type for triggering when new actors are placed on to the viewport. Triggers before NewActorsDropped if placement is caused by a drop action */
 	DECLARE_MULTICAST_DELEGATE_TwoParams(FOnNewActorsPlaced, UObject*, const TArray<AActor*>&);
+	/** delegate type for triggering when an actor is replaced by another one in the editor. Triggers once AActor::EditorReplacedActor is completed */
+	DECLARE_MULTICAST_DELEGATE_TwoParams(FOnEditorActorReplaced, AActor*, AActor*);
 	/** delegate type for when attempting to apply an object to an actor */
 	DECLARE_MULTICAST_DELEGATE_TwoParams(FOnApplyObjectToActor, UObject*, AActor*);
 	/** delegate type for triggering when grid snapping has changed */
@@ -290,6 +292,8 @@ struct FEditorDelegates
 	static UNREALED_API FOnNewActorsDropped OnNewActorsDropped;
 	/** Called when new actors are placed in the viewport */
 	static UNREALED_API FOnNewActorsPlaced OnNewActorsPlaced;
+	/** Called when an actor is replaced in the editor (once AActor::EditorReplacedActor is completed). The first parameter is the old actor, the second is the new actor. */
+	static UNREALED_API FOnEditorActorReplaced OnEditorActorReplaced;
 	/** Called when grid snapping is changed */
 	static UNREALED_API FOnGridSnappingChanged OnGridSnappingChanged;
 	/** Called when a lighting build has started */
diff --git a/Engine/Source/Editor/WorldPartitionEditor/Private/WorldPartition/ContentBundle/ContentBundleEditorSubsystem.cpp b/Engine/Source/Editor/WorldPartitionEditor/Private/WorldPartition/ContentBundle/ContentBundleEditorSubsystem.cpp
--- a/Engine/Source/Editor/WorldPartitionEditor/Private/WorldPartition/ContentBundle/ContentBundleEditorSubsystem.cpp
+++ b/Engine/Source/Editor/WorldPartitionEditor/Private/WorldPartition/ContentBundle/ContentBundleEditorSubsystem.cpp
@@ -1,18 +1,18 @@
 // Copyright Epic Games, Inc. All Rights Reserved.
 #include "WorldPartition/ContentBundle/ContentBundleEditorSubsystem.h"
 #include "WorldPartition/ContentBundle/ContentBundleEngineSubsystem.h"
-#include "WorldPartition/ContentBundle/ContentBundle.h"
+#include "WorldPartition/ContentBundle/ContentBundleWorldSubsystem.h"
 #include "WorldPartition/ContentBundle/ContentBundleDescriptor.h"
-#include "WorldPartition/ContentBundle/ContentBundle.h"
 #include "WorldPartition/ContentBundle/ContentBundleEditor.h"
+#include "WorldPartition/ContentBundle/ContentBundleLog.h"
+#include "WorldPartition/ContentBundle/ContentBundle.h"
 #include "WorldPartition/DataLayer/ExternalDataLayerManager.h"
 #include "WorldPartition/DataLayer/ExternalDataLayerInstance.h"
 #include "WorldPartition/DataLayer/ExternalDataLayerAsset.h"
 #include "WorldPartition/WorldPartition.h"
 #include "Subsystems/ActorEditorContextSubsystem.h"
+#include "DataLayer/DataLayerEditorSubsystem.h"
 #include "Engine/Selection.h"
-#include "WorldPartition/ContentBundle/ContentBundleWorldSubsystem.h"
-#include "WorldPartition/ContentBundle/ContentBundleLog.h"
 #include "Styling/AppStyle.h"
 #include "Widgets/Images/SImage.h"
 #include "Widgets/SBoxPanel.h"
@@ -80,9 +80,7 @@ void UContentBundleEditingSubmodule::ApplyContext(AActor* InActor)
 	}
 
 	// Prefer override spawning External Data Layer over Content Bundles
-	const UExternalDataLayerAsset* ExternalDataLayerAsset = Cast<UExternalDataLayerAsset>(ULevel::GetOverrideSpawningLevelMountPointObject());
-	const UExternalDataLayerManager* ExternalDataLayerManager = UExternalDataLayerManager::GetExternalDataLayerManager(InActor);
-	if (const UExternalDataLayerInstance* OverrideSpawningExternalDataLayerInstance = (ExternalDataLayerAsset && ExternalDataLayerManager) ? ExternalDataLayerManager->GetExternalDataLayerInstance(ExternalDataLayerAsset) : nullptr)
+	if (const UExternalDataLayerInstance* SpawningExternalDataLayerInstance = ULevel::GetOverrideSpawningLevelMountPointObject() ? UDataLayerEditorSubsystem::Get()->GetActorSpawningExternalDataLayerInstance(InActor) : nullptr)
 	{
 		return;
 	}
@@ -98,7 +96,7 @@ void UContentBundleEditingSubmodule::OnExecuteActorEditorContextAction(UWorld* I
 	switch (InType)
 	{
 	case EActorEditorContextAction::ApplyContext:
-		ApplyContext(InActor);
+		// For backward compatibility, Content Bundle will only apply context trough OnLevelActorAdded
 		break;
 	case EActorEditorContextAction::ResetContext:
 		if (TSharedPtr<FContentBundleEditor> EditingContentBundle = GetEditorContentBundle(EditingContentBundleGuid))
diff --git a/Engine/Source/Runtime/Engine/Classes/GameFramework/Actor.h b/Engine/Source/Runtime/Engine/Classes/GameFramework/Actor.h
--- a/Engine/Source/Runtime/Engine/Classes/GameFramework/Actor.h
+++ b/Engine/Source/Runtime/Engine/Classes/GameFramework/Actor.h
@@ -4548,6 +4548,8 @@ private:
 	ENGINE_API static bool AddDataLayerAsset(AActor* InActor, const UDataLayerAsset* InDataLayerAsset);
 	ENGINE_API static bool RemoveDataLayerAsset(AActor* InActor, const UDataLayerAsset* InDataLayerAsset);
 
+	friend class UEngine;
+	friend class FContentBundleEditor;
 	friend class UDataLayerInstanceWithAsset;
 	friend class UDataLayerInstancePrivate;
 	friend class UExternalDataLayerInstance;
diff --git a/Engine/Source/Runtime/Engine/Private/ActorEditor.cpp b/Engine/Source/Runtime/Engine/Private/ActorEditor.cpp
--- a/Engine/Source/Runtime/Engine/Private/ActorEditor.cpp
+++ b/Engine/Source/Runtime/Engine/Private/ActorEditor.cpp
@@ -30,6 +30,7 @@
 #include "WorldPartition/DataLayer/DeprecatedDataLayerInstance.h"
 #include "WorldPartition/DataLayer/ExternalDataLayerInstance.h"
 #include "WorldPartition/DataLayer/ExternalDataLayerAsset.h"
+#include "WorldPartition/ContentBundle/ContentBundlePaths.h"
 #include "LevelInstance/LevelInstanceSubsystem.h"
 #include "ActorFolder.h"
 #include "WorldPersistentFolders.h"
@@ -1493,7 +1494,13 @@ void AActor::SetFolderPath_Recursively(const FName& NewFolderPath)
 // Ideally, this should be revisited to implement something more generic.
 void AActor::EditorReplacedActor(AActor* OldActor)
 {
-	ContentBundleGuid = OldActor->ContentBundleGuid;
+	// Don't update Content Bundle if new actor is using External Data Layer
+	if (!ExternalDataLayerAsset)
+	{
+		// We can't rely on OldActor to transfer the ContentBundleGuid as it could have moved to a different Content Bundle
+		// Resolve ContentBundleGuid from ActorPackage
+		ContentBundleGuid = ContentBundlePaths::GetContentBundleGuidFromExternalActorPackagePath(GetPackage()->GetFName().ToString());
+	}
 
 	SetActorLabel(OldActor->GetActorLabel());
 	Tags = OldActor->Tags;
@@ -1512,11 +1519,15 @@ void AActor::EditorReplacedActor(AActor* OldActor)
 	TArray<const UDataLayerInstance*> ConstDataLayerInstances = OldActor->GetDataLayerInstancesInternal(bUseLevelContext, bIncludeParentDataLayers);
 	if (!ConstDataLayerInstances.IsEmpty())
 	{
+		// Transfer old data layers (except the external data layer as it is handled at actor spawning)
 		TArray<UDataLayerInstance*> DataLayerInstances;
-		Algo::Transform(ConstDataLayerInstances, DataLayerInstances, [](const UDataLayerInstance* ConstDataLayerInstance) { return const_cast<UDataLayerInstance*>(ConstDataLayerInstance); });
+		Algo::TransformIf(ConstDataLayerInstances, DataLayerInstances, [](const UDataLayerInstance* ConstDataLayerInstance) { return !ConstDataLayerInstance->IsA<UExternalDataLayerInstance>(); }, [](const UDataLayerInstance* ConstDataLayerInstance) { return const_cast<UDataLayerInstance*>(ConstDataLayerInstance); });
 		IDataLayerEditorModule& EditorModule = FModuleManager::LoadModuleChecked<IDataLayerEditorModule>("DataLayerEditor");
 		EditorModule.AddActorToDataLayers(this, DataLayerInstances);
 	}
+
+	// Broadcast that actor was replaced
+	FEditorDelegates::OnEditorActorReplaced.Broadcast(OldActor, this);
 }
 
 void AActor::CheckForDeprecated()
diff --git a/Engine/Source/Runtime/Engine/Private/IWorldPartitionEditorModule.cpp b/Engine/Source/Runtime/Engine/Private/IWorldPartitionEditorModule.cpp
new file mode 100644
--- /dev/null
+++ b/Engine/Source/Runtime/Engine/Private/IWorldPartitionEditorModule.cpp
@@ -0,0 +1,15 @@
+// Copyright Epic Games, Inc. All Rights Reserved.
+
+#include "WorldPartition/IWorldPartitionEditorModule.h"
+#include "Modules/ModuleManager.h"
+
+IWorldPartitionEditorModule* IWorldPartitionEditorModule::Instance = nullptr;
+
+IWorldPartitionEditorModule& IWorldPartitionEditorModule::Get()
+{
+	if (!Instance)
+	{
+		Instance = &FModuleManager::LoadModuleChecked<IWorldPartitionEditorModule>("WorldPartitionEditor");
+	}
+	return *Instance;
+}
\ No newline at end of file
diff --git a/Engine/Source/Runtime/Engine/Private/UnrealEngine.cpp b/Engine/Source/Runtime/Engine/Private/UnrealEngine.cpp
--- a/Engine/Source/Runtime/Engine/Private/UnrealEngine.cpp
+++ b/Engine/Source/Runtime/Engine/Private/UnrealEngine.cpp
@@ -152,6 +152,7 @@ UnrealEngine.cpp: Implements the UEngine class and helpers.
 #include "TextureCompiler.h"
 #include "GenericPlatform/GenericPlatformCrashContext.h"
 #include "Framework/Notifications/NotificationManager.h"
+#include "WorldPartition/DataLayer/ExternalDataLayerAsset.h"
 #endif
 // @todo this is here only due to circular dependency to AIModule. To be removed
 
@@ -17366,6 +17367,7 @@ void UEngine::CopyPropertiesForUnrelatedObjects(UObject* OldObject, UObject* New
 #if WITH_EDITOR
 	FGuid NewActorGuid;
 	FGuid NewActorInstanceGuid;
+	const UExternalDataLayerAsset* NewExternalDataLayerAsset = nullptr;
 #endif
 
 	// Bad idea to write data to an actor while its components are registered
@@ -17380,6 +17382,7 @@ void UEngine::CopyPropertiesForUnrelatedObjects(UObject* OldObject, UObject* New
 #if WITH_EDITOR
 		NewActorGuid = NewActor->GetActorGuid();
 		NewActorInstanceGuid = NewActor->GetActorInstanceGuid();
+		NewExternalDataLayerAsset = NewActor->GetExternalDataLayerAsset();
 #endif
 	}
 
@@ -17613,6 +17616,20 @@ void UEngine::CopyPropertiesForUnrelatedObjects(UObject* OldObject, UObject* New
 		{
 			FSetActorInstanceGuid SetActorInstanceGuid(NewActor, NewActorInstanceGuid);
 		}
+
+		// Restore External Data Layer Asset
+		const UExternalDataLayerAsset* ExternalDataLayerAsset = NewActor->GetExternalDataLayerAsset();
+		if (!GIsReinstancing  && (ExternalDataLayerAsset != NewExternalDataLayerAsset))
+		{
+			if (ExternalDataLayerAsset)
+			{
+				FAssignActorDataLayer::RemoveDataLayerAsset(NewActor, ExternalDataLayerAsset);
+			}
+			if (NewExternalDataLayerAsset)
+			{
+				FAssignActorDataLayer::AddDataLayerAsset(NewActor, NewExternalDataLayerAsset);
+			}
+		}
 #endif
 	}
 
diff --git a/Engine/Source/Runtime/Engine/Private/WorldPartition/DataLayer/ExternalDataLayerEngineSubsystem.cpp b/Engine/Source/Runtime/Engine/Private/WorldPartition/DataLayer/ExternalDataLayerEngineSubsystem.cpp
--- a/Engine/Source/Runtime/Engine/Private/WorldPartition/DataLayer/ExternalDataLayerEngineSubsystem.cpp
+++ b/Engine/Source/Runtime/Engine/Private/WorldPartition/DataLayer/ExternalDataLayerEngineSubsystem.cpp
@@ -115,21 +115,7 @@ void UExternalDataLayerEngineSubsystem::OnGetLevelExternalActorsPaths(const FStr
 
 bool UExternalDataLayerEngineSubsystem::OnResolveLevelMountPoint(const FString& InLevelPackageName, const UObject* InLevelMountPointContext, FString& OutResolvedLevelMountPoint)
 {
-	const UExternalDataLayerAsset* ExternalDataLayerAssetContext = nullptr;
-	if (InLevelMountPointContext)
-	{
-		ExternalDataLayerAssetContext = Cast<UExternalDataLayerAsset>(InLevelMountPointContext);
-		if (!ExternalDataLayerAssetContext && InLevelMountPointContext->Implements<UDataLayerInstanceProvider>())
-		{
-			ExternalDataLayerAssetContext = CastChecked<IDataLayerInstanceProvider>(InLevelMountPointContext)->GetRootExternalDataLayerAsset();
-		}
-		if (!ExternalDataLayerAssetContext && InLevelMountPointContext->IsA<AActor>())
-		{
-			ExternalDataLayerAssetContext = CastChecked<AActor>(InLevelMountPointContext)->GetExternalDataLayerAsset();
-		}
-	}
-
-	if (ExternalDataLayerAssetContext)
+	if (const UExternalDataLayerAsset* ExternalDataLayerAssetContext = FExternalDataLayerHelper::GetExternalDataLayerAssetFromObject(InLevelMountPointContext))
 	{
 		OutResolvedLevelMountPoint = FExternalDataLayerHelper::GetExternalDataLayerLevelRootPath(ExternalDataLayerAssetContext, InLevelPackageName);
 		return true;
diff --git a/Engine/Source/Runtime/Engine/Private/WorldPartition/DataLayer/ExternalDataLayerHelper.cpp b/Engine/Source/Runtime/Engine/Private/WorldPartition/DataLayer/ExternalDataLayerHelper.cpp
--- a/Engine/Source/Runtime/Engine/Private/WorldPartition/DataLayer/ExternalDataLayerHelper.cpp
+++ b/Engine/Source/Runtime/Engine/Private/WorldPartition/DataLayer/ExternalDataLayerHelper.cpp
@@ -352,6 +352,24 @@ bool FExternalDataLayerHelper::IsExternalDataLayerPath(FStringView InExternalDat
 	return false;
 }
 
+const UExternalDataLayerAsset* FExternalDataLayerHelper::GetExternalDataLayerAssetFromObject(const UObject* InContextObject)
+{
+	const UExternalDataLayerAsset* ExternalDataLayerAssetContext = nullptr;
+	if (InContextObject)
+	{
+		ExternalDataLayerAssetContext = Cast<UExternalDataLayerAsset>(InContextObject);
+		if (!ExternalDataLayerAssetContext && InContextObject->Implements<UDataLayerInstanceProvider>())
+		{
+			ExternalDataLayerAssetContext = CastChecked<IDataLayerInstanceProvider>(InContextObject)->GetRootExternalDataLayerAsset();
+		}
+		if (!ExternalDataLayerAssetContext && InContextObject->IsA<AActor>())
+		{
+			ExternalDataLayerAssetContext = CastChecked<AActor>(InContextObject)->GetExternalDataLayerAsset();
+		}
+	}
+	return ExternalDataLayerAssetContext;
+}
+
 #endif
 
 #undef LOCTEXT_NAMESPACE
\ No newline at end of file
diff --git a/Engine/Source/Runtime/Engine/Private/WorldPartition/DataLayer/ExternalDataLayerManager.cpp b/Engine/Source/Runtime/Engine/Private/WorldPartition/DataLayer/ExternalDataLayerManager.cpp
--- a/Engine/Source/Runtime/Engine/Private/WorldPartition/DataLayer/ExternalDataLayerManager.cpp
+++ b/Engine/Source/Runtime/Engine/Private/WorldPartition/DataLayer/ExternalDataLayerManager.cpp
@@ -531,7 +531,7 @@ const UExternalDataLayerAsset* UExternalDataLayerManager::GetActorEditorContextC
 	return nullptr;
 }
 
-const UExternalDataLayerAsset* UExternalDataLayerManager::GetMatchingExternalDataLayerAssetForObjectPath(const FSoftObjectPath& InObjectPath)
+const UExternalDataLayerAsset* UExternalDataLayerManager::GetMatchingExternalDataLayerAssetForObjectPath(const FSoftObjectPath& InObjectPath) const
 {
 	const FName MountPoint = FPackageName::GetPackageMountPoint(InObjectPath.ToString());
 
diff --git a/Engine/Source/Runtime/Engine/Private/WorldPartition/DataLayer/WorldDataLayers.cpp b/Engine/Source/Runtime/Engine/Private/WorldPartition/DataLayer/WorldDataLayers.cpp
--- a/Engine/Source/Runtime/Engine/Private/WorldPartition/DataLayer/WorldDataLayers.cpp
+++ b/Engine/Source/Runtime/Engine/Private/WorldPartition/DataLayer/WorldDataLayers.cpp
@@ -24,6 +24,7 @@
 #include UE_INLINE_GENERATED_CPP_BY_NAME(WorldDataLayers)
 
 #if WITH_EDITOR
+#include "ActorEditorContext/ScopedActorEditorContextSetExternalDataLayerAsset.h"
 #include "WorldPartition/WorldPartitionEditorPerProjectUserSettings.h"
 #include "WorldPartition/DataLayer/WorldDataLayersActorDesc.h"
 #include "WorldPartition/ContentBundle/ContentBundleEditorSubsystemInterface.h"
@@ -689,6 +690,9 @@ AWorldDataLayers* AWorldDataLayers::Create(const FActorSpawnParameters& SpawnPar
 
 	if (!WorldDataLayers)
 	{
+		// Make sure there's no context while creating the AWorldDataLayers (avoids generating any warnings while spawning)
+		FScopedActorEditorContextSetExternalDataLayerAsset ScopeDisableCurrentEDL(nullptr);
+		FScopedOverrideSpawningLevelMountPointObject ScopeDisableOverrideEDL(nullptr);
 		UWorld* World = SpawnParameters.OverrideLevel->GetWorld();
 		WorldDataLayers = World->SpawnActor<AWorldDataLayers>(AWorldDataLayers::StaticClass(), SpawnParameters);
 	}
diff --git a/Engine/Source/Runtime/Engine/Public/WorldPartition/DataLayer/ExternalDataLayerHelper.h b/Engine/Source/Runtime/Engine/Public/WorldPartition/DataLayer/ExternalDataLayerHelper.h
--- a/Engine/Source/Runtime/Engine/Public/WorldPartition/DataLayer/ExternalDataLayerHelper.h
+++ b/Engine/Source/Runtime/Engine/Public/WorldPartition/DataLayer/ExternalDataLayerHelper.h
@@ -52,6 +52,8 @@ public:
 private:
 
 #if WITH_EDITOR
+	ENGINE_API static const UExternalDataLayerAsset* GetExternalDataLayerAssetFromObject(const UObject* InContextObject);
+
 	/** 
 	 * Validates that all actors can change their External Data Layer to the new provided value (supports passing null) 
 	 * Returns false if any actor fails and fills OutFailureReason with the reason (if non-null). 
@@ -66,6 +68,8 @@ private:
 
 	friend class FDataLayerEditorModule;
 	friend class UDataLayerEditorSubsystem;
+	friend class UExternalDataLayerEngineSubsystem;
+	friend class UContentBundleEditingSubmodule;
 #endif
 	static constexpr FStringView GetExternalDataLayerFolder() { return ExternalDataLayerFolder; }
 	static constexpr FStringView ExternalDataLayerFolder = TEXTVIEW("/EDL/");
diff --git a/Engine/Source/Runtime/Engine/Public/WorldPartition/DataLayer/ExternalDataLayerManager.h b/Engine/Source/Runtime/Engine/Public/WorldPartition/DataLayer/ExternalDataLayerManager.h
--- a/Engine/Source/Runtime/Engine/Public/WorldPartition/DataLayer/ExternalDataLayerManager.h
+++ b/Engine/Source/Runtime/Engine/Public/WorldPartition/DataLayer/ExternalDataLayerManager.h
@@ -78,7 +78,7 @@ private:
 	bool ValidateOnActorExternalDataLayerAssetChanged(AActor* InActor);
 	bool RegisterExternalDataLayerInstance(UExternalDataLayerInstance* InExternalDataLayerInstance);
 	bool UnregisterExternalDataLayerInstance(UExternalDataLayerInstance* InExternalDataLayerInstance);
-	const UExternalDataLayerAsset* GetMatchingExternalDataLayerAssetForObjectPath(const FSoftObjectPath& InObjectPath);
+	const UExternalDataLayerAsset* GetMatchingExternalDataLayerAssetForObjectPath(const FSoftObjectPath& InObjectPath) const;
 	const UExternalDataLayerAsset* GetActorEditorContextCurrentExternalDataLayer() const;
 	AWorldDataLayers* GetWorldDataLayers(const UExternalDataLayerAsset* InExternalDataLayerAsset, bool bInAllowCreate = false) const;
 	FString GetActorPackageName(const UExternalDataLayerAsset* InExternalDataLayerAsset, const ULevel* InDestinationLevel, const FString& InActorPath) const;
diff --git a/Engine/Source/Runtime/Engine/Public/WorldPartition/IWorldPartitionEditorModule.h b/Engine/Source/Runtime/Engine/Public/WorldPartition/IWorldPartitionEditorModule.h
--- a/Engine/Source/Runtime/Engine/Public/WorldPartition/IWorldPartitionEditorModule.h
+++ b/Engine/Source/Runtime/Engine/Public/WorldPartition/IWorldPartitionEditorModule.h
@@ -15,6 +15,8 @@ class IWorldPartitionEditorModule : public IModuleInterface
 public:
 	virtual ~IWorldPartitionEditorModule() {}
 
+	static ENGINE_API IWorldPartitionEditorModule& Get();
+
 	virtual bool ConvertMap(const FString& InLongPackageName) = 0;
 
 	UE_DEPRECATED(5.1, "Use RunBuilder with UWorld* instead.")
@@ -90,4 +92,7 @@ public:
 
 	/** Return the commandlet post-execution event */
 	virtual FOnPostExecuteCommandlet& OnPostExecuteCommandlet() = 0;
+
+private:
+	static IWorldPartitionEditorModule* Instance;
 };
\ No newline at end of file
