Fixed CIS static analysis

diff --git a/Engine/Source/Editor/DataLayerEditor/Private/DataLayer/DataLayerEditorSubsystem.cpp b/Engine/Source/Editor/DataLayerEditor/Private/DataLayer/DataLayerEditorSubsystem.cpp
--- a/Engine/Source/Editor/DataLayerEditor/Private/DataLayer/DataLayerEditorSubsystem.cpp
+++ b/Engine/Source/Editor/DataLayerEditor/Private/DataLayer/DataLayerEditorSubsystem.cpp
@@ -417,7 +417,7 @@ const UExternalDataLayerInstance* UDataLayerEditorSubsystem::GetActorSpawningExt
 		ResolvedExternalDataLayerAsset = GetActorEditorContextCurrentExternalDataLayer();
 	}
 
-	const UExternalDataLayerInstance* ResolvedExternalDataLayerInstance = ResolvedExternalDataLayerAsset ? ExternalDataLayerManager->GetExternalDataLayerInstance(ResolvedExternalDataLayerAsset) : nullptr;
+	const UExternalDataLayerInstance* ResolvedExternalDataLayerInstance = (ResolvedExternalDataLayerAsset && ExternalDataLayerManager) ? ExternalDataLayerManager->GetExternalDataLayerInstance(ResolvedExternalDataLayerAsset) : nullptr;
 	return ResolvedExternalDataLayerInstance;
 }
 
