Fixed UDataLayerEditorSubsystem trying to access the DataLayerManager for a non-partitioned level, but using OFPA

diff --git a/Engine/Source/Editor/DataLayerEditor/Private/DataLayer/DataLayerEditorSubsystem.cpp b/Engine/Source/Editor/DataLayerEditor/Private/DataLayer/DataLayerEditorSubsystem.cpp
--- a/Engine/Source/Editor/DataLayerEditor/Private/DataLayer/DataLayerEditorSubsystem.cpp
+++ b/Engine/Source/Editor/DataLayerEditor/Private/DataLayer/DataLayerEditorSubsystem.cpp
@@ -517,6 +517,12 @@ void UDataLayerEditorSubsystem::ApplyContext(AActor* InActor, bool bInForceTryAp
 		return;
 	}
 
+	UDataLayerManager* DataLayerManager = UDataLayerManager::GetDataLayerManager(InActor->GetWorld());
+	if (!DataLayerManager)
+	{
+		return;
+	}
+
 	// Try to apply context External Data Layer (this operation can fail if asset referencing validation fails)
 	const UExternalDataLayerAsset* CurrentExternalDataLayer = GetActorEditorContextCurrentExternalDataLayer();
 	const UExternalDataLayerAsset* ActorExternalDataLayerAsset = InActor->GetExternalDataLayerAsset();
@@ -540,7 +546,6 @@ void UDataLayerEditorSubsystem::ApplyContext(AActor* InActor, bool bInForceTryAp
 	}
 
 	// Apply context Data Layers (except External Data Layer)
-	UDataLayerManager* DataLayerManager = UDataLayerManager::GetDataLayerManager(InActor->GetWorld());
 	TArray<UDataLayerInstance*> DataLayerInstances = DataLayerManager->GetActorEditorContextDataLayers();
 	if (!DataLayerInstances.IsEmpty())
 	{
