World Partition - Data Layer: Prevent non-user managed actors from being added/removed to/from data layers using the buttons next to the search box in the DL Outliner

diff --git a/Engine/Source/Editor/DataLayerEditor/Private/DataLayer/SDataLayerOutliner.cpp b/Engine/Source/Editor/DataLayerEditor/Private/DataLayer/SDataLayerOutliner.cpp
--- a/Engine/Source/Editor/DataLayerEditor/Private/DataLayer/SDataLayerOutliner.cpp
+++ b/Engine/Source/Editor/DataLayerEditor/Private/DataLayer/SDataLayerOutliner.cpp
@@ -2,6 +2,7 @@
 
 #include "SDataLayerOutliner.h"
 
+#include "Algo/AllOf.h"
 #include "Algo/Transform.h"
 #include "DataLayer/DataLayerEditorSubsystem.h"
 #include "DataLayer/DataLayerTreeItem.h"
@@ -11,6 +12,7 @@
 #include "Internationalization/Internationalization.h"
 #include "Misc/Attribute.h"
 #include "ScopedTransaction.h"
+#include "Selection.h"
 #include "SlotBase.h"
 #include "Styling/AppStyle.h"
 #include "Styling/ISlateStyle.h"
@@ -32,7 +34,7 @@ void SDataLayerOutliner::CustomAddToToolbar(TSharedPtr<SHorizontalBox> Toolbar)
 		.Padding(4.f, 0.f, 0.f, 0.f)
 		[
 			SNew(SButton)
-			.IsEnabled(this, &SDataLayerOutliner::CanAddSelectedActorsToSelectedDataLayersClicked)
+			.IsEnabled(this, &SDataLayerOutliner::CanAddSelectedActorsToSelectedDataLayers)
 			.ButtonStyle(FAppStyle::Get(), "SimpleButton")
 			.ToolTipText(LOCTEXT("AddSelectedActorsToSelectedDataLayersTooltip", "Add selected actors to selected Data Layers"))
 			.OnClicked(this, &SDataLayerOutliner::OnAddSelectedActorsToSelectedDataLayersClicked)
@@ -48,7 +50,7 @@ void SDataLayerOutliner::CustomAddToToolbar(TSharedPtr<SHorizontalBox> Toolbar)
 		.AutoWidth()
 		[
 			SNew(SButton)
-			.IsEnabled(this, &SDataLayerOutliner::CanRemoveSelectedActorsFromSelectedDataLayersClicked)
+			.IsEnabled(this, &SDataLayerOutliner::CanRemoveSelectedActorsFromSelectedDataLayers)
 			.ButtonStyle(FAppStyle::Get(), "SimpleButton")
 			.ToolTipText(LOCTEXT("RemoveSelectedActorsFromSelectedDataLayersTooltip", "Remove selected actors from selected Data Layers"))
 			.OnClicked(this, &SDataLayerOutliner::OnRemoveSelectedActorsFromSelectedDataLayersClicked)
@@ -70,25 +72,33 @@ TArray<UDataLayerInstance*> SDataLayerOutliner::GetSelectedDataLayers() const
 	return ValidSelectedDataLayers;
 }
 
-bool SDataLayerOutliner::CanAddSelectedActorsToSelectedDataLayersClicked() const
+bool SDataLayerOutliner::CanAddSelectedActorsToSelectedDataLayers() const
 {
 	if (GEditor->GetSelectedActorCount() > 0)
 	{
 		TArray<UDataLayerInstance*> SelectedDataLayerInstances = GetSelectedDataLayers();
-		const bool bSelectedDataLayerInstancesContainsReadOnly = !!SelectedDataLayerInstances.FindByPredicate([](const UDataLayerInstance* DataLayerInstance) { return DataLayerInstance->IsReadOnly(); });
-		return (!SelectedDataLayerInstances.IsEmpty() && !bSelectedDataLayerInstancesContainsReadOnly);
+		if (!SelectedDataLayerInstances.IsEmpty())
+		{
+			TArray<AActor*> SelectedActors;
+			GEditor->GetSelectedActors()->GetSelectedObjects<AActor>(SelectedActors);
+
+			const bool bSelectedActorsAreAllUserManaged = Algo::AllOf(SelectedActors, [](const AActor* Actor) { return Actor->IsUserManaged(); });
+			const bool bAllSelectedDataLayersCanAddActors = Algo::AllOf(SelectedDataLayerInstances, [](const UDataLayerInstance* DataLayerInstance) { return DataLayerInstance->CanUserAddActors(); });
+
+			return bAllSelectedDataLayersCanAddActors && bSelectedActorsAreAllUserManaged;
+		}
 	}
 	return false;
 }
 
-bool SDataLayerOutliner::CanRemoveSelectedActorsFromSelectedDataLayersClicked() const
+bool SDataLayerOutliner::CanRemoveSelectedActorsFromSelectedDataLayers() const
 {
-	return CanAddSelectedActorsToSelectedDataLayersClicked();
+	return CanAddSelectedActorsToSelectedDataLayers();
 }
 
 FReply SDataLayerOutliner::OnAddSelectedActorsToSelectedDataLayersClicked()
 {
-	if (CanAddSelectedActorsToSelectedDataLayersClicked())
+	if (CanAddSelectedActorsToSelectedDataLayers())
 	{
 		TArray<UDataLayerInstance*> SelectedDataLayers = GetSelectedDataLayers();
 		const FScopedTransaction Transaction(LOCTEXT("AddSelectedActorsToSelectedDataLayers", "Add Selected Actor(s) to Selected Data Layer(s)"));
@@ -99,7 +109,7 @@ FReply SDataLayerOutliner::OnAddSelectedActorsToSelectedDataLayersClicked()
 
 FReply SDataLayerOutliner::OnRemoveSelectedActorsFromSelectedDataLayersClicked()
 {
-	if (CanRemoveSelectedActorsFromSelectedDataLayersClicked())
+	if (CanRemoveSelectedActorsFromSelectedDataLayers())
 	{
 		TArray<UDataLayerInstance*> SelectedDataLayers = GetSelectedDataLayers();
 		const FScopedTransaction Transaction(LOCTEXT("RemoveSelectedActorsFromSelectedDataLayers", "Remove Selected Actors from Selected Data Layers"));
diff --git a/Engine/Source/Editor/DataLayerEditor/Private/DataLayer/SDataLayerOutliner.h b/Engine/Source/Editor/DataLayerEditor/Private/DataLayer/SDataLayerOutliner.h
--- a/Engine/Source/Editor/DataLayerEditor/Private/DataLayer/SDataLayerOutliner.h
+++ b/Engine/Source/Editor/DataLayerEditor/Private/DataLayer/SDataLayerOutliner.h
@@ -22,8 +22,8 @@ public:
 
 	void CustomAddToToolbar(TSharedPtr<class SHorizontalBox> Toolbar) override;
 private:
-	bool CanAddSelectedActorsToSelectedDataLayersClicked() const;
-	bool CanRemoveSelectedActorsFromSelectedDataLayersClicked() const;
+	bool CanAddSelectedActorsToSelectedDataLayers() const;
+	bool CanRemoveSelectedActorsFromSelectedDataLayers() const;
 	FReply OnAddSelectedActorsToSelectedDataLayersClicked();
 	FReply OnRemoveSelectedActorsFromSelectedDataLayersClicked();
 	TArray<class UDataLayerInstance*> GetSelectedDataLayers() const;
