Added UDataLayerEditorSubsystem::SetParentDataLayerForDataLayers to optimize case where multiple data layer instances are moved under a parent data layer instance.

diff --git a/Engine/Source/Editor/DataLayerEditor/Private/DataLayer/DataLayerEditorSubsystem.cpp b/Engine/Source/Editor/DataLayerEditor/Private/DataLayer/DataLayerEditorSubsystem.cpp
--- a/Engine/Source/Editor/DataLayerEditor/Private/DataLayer/DataLayerEditorSubsystem.cpp
+++ b/Engine/Source/Editor/DataLayerEditor/Private/DataLayer/DataLayerEditorSubsystem.cpp
@@ -817,21 +817,44 @@ UWorld* UDataLayerEditorSubsystem::GetWorld() const
 	return GEditor->GetEditorWorldContext().World();
 }
 
-bool UDataLayerEditorSubsystem::SetParentDataLayer(UDataLayerInstance* DataLayerInstance, UDataLayerInstance* ParentDataLayer)
+void UDataLayerEditorSubsystem::SetParentDataLayerForDataLayers(const TArray<UDataLayerInstance*>& DataLayers, UDataLayerInstance* ParentDataLayer)
 {
-	if (DataLayerInstance->CanBeChildOf(ParentDataLayer))
+	SetParentDataLayerForDataLayersInternal(DataLayers, ParentDataLayer);
+}
+
+bool UDataLayerEditorSubsystem::SetParentDataLayerForDataLayersInternal(const TArray<UDataLayerInstance*>& DataLayers, UDataLayerInstance* ParentDataLayer)
+{
+	bool bResult = false;
+	bool bLoadingStateChanged = false;
+
+	for (UDataLayerInstance* DataLayerInstance : DataLayers)
+	{
+		if (DataLayerInstance->CanBeChildOf(ParentDataLayer))
+		{
+			const bool bIsLoaded = DataLayerInstance->IsEffectiveLoadedInEditor();
+			DataLayerInstance->SetParent(ParentDataLayer);
+			if (bIsLoaded != DataLayerInstance->IsEffectiveLoadedInEditor())
+			{
+				bLoadingStateChanged = true;
+			}
+			bResult = true;
+		}
+	}
+	if (bResult)
 	{
-		const bool bIsLoaded = DataLayerInstance->IsEffectiveLoadedInEditor();
-		DataLayerInstance->SetParent(ParentDataLayer);
 		BroadcastDataLayerChanged(EDataLayerAction::Reset, NULL, NAME_None);
 		UpdateAllActorsVisibility(true, true);
-		if (bIsLoaded != DataLayerInstance->IsEffectiveLoadedInEditor())
+		if (bLoadingStateChanged)
 		{
 			OnDataLayerEditorLoadingStateChanged(true);
 		}
-		return true;
 	}
-	return false;
+	return bResult;
+}
+
+bool UDataLayerEditorSubsystem::SetParentDataLayer(UDataLayerInstance* DataLayerInstance, UDataLayerInstance* ParentDataLayer)
+{
+	return SetParentDataLayerForDataLayersInternal({ DataLayerInstance }, ParentDataLayer);
 }
 
 bool UDataLayerEditorSubsystem::AddActorToDataLayer(AActor* Actor, UDataLayerInstance* DataLayerInstance)
diff --git a/Engine/Source/Editor/DataLayerEditor/Private/DataLayer/DataLayerMode.cpp b/Engine/Source/Editor/DataLayerEditor/Private/DataLayer/DataLayerMode.cpp
--- a/Engine/Source/Editor/DataLayerEditor/Private/DataLayer/DataLayerMode.cpp
+++ b/Engine/Source/Editor/DataLayerEditor/Private/DataLayer/DataLayerMode.cpp
@@ -1166,10 +1166,7 @@ void FDataLayerMode::SetParentDataLayer(const TArray<UDataLayerInstance*> DataLa
 		if (!ValidDataLayers.IsEmpty())
 		{
 			const FScopedTransaction Transaction(LOCTEXT("DataLayerOutlinerChangeDataLayersParent", "Change Data Layers Parent"));
-			for (UDataLayerInstance* DataLayerInstance : ValidDataLayers)
-			{
-				DataLayerEditorSubsystem->SetParentDataLayer(DataLayerInstance, ParentDataLayer);
-			}
+			DataLayerEditorSubsystem->SetParentDataLayerForDataLayers(ValidDataLayers, ParentDataLayer);
 		}
 	}
 }
diff --git a/Engine/Source/Editor/DataLayerEditor/Public/DataLayer/DataLayerEditorSubsystem.h b/Engine/Source/Editor/DataLayerEditor/Public/DataLayer/DataLayerEditorSubsystem.h
--- a/Engine/Source/Editor/DataLayerEditor/Public/DataLayer/DataLayerEditorSubsystem.h
+++ b/Engine/Source/Editor/DataLayerEditor/Public/DataLayer/DataLayerEditorSubsystem.h
@@ -611,6 +611,15 @@ public:
 	UFUNCTION(BlueprintCallable, Category = DataLayers)
 	bool SetParentDataLayer(UDataLayerInstance* DataLayer, UDataLayerInstance* ParentDataLayer);
 
+	/**
+	 * Sets a Parent DataLayer for a specified list of DataLayers
+	 * 
+	 *  @param DataLayers		The child DataLayers.
+	 *  @param ParentDataLayer	The parent DataLayer.
+	 */
+	UFUNCTION(BlueprintCallable, Category = DataLayers)
+	void SetParentDataLayerForDataLayers(const TArray<UDataLayerInstance*>& DataLayers, UDataLayerInstance* ParentDataLayer);
+
 	/**
 	 * Deletes all of the provided DataLayers
 	 *
@@ -755,6 +764,7 @@ private:
 	void OnSelectionChanged();
 	void RebuildSelectedDataLayersFromEditorSelection();
 	const TSet<TWeakObjectPtr<const UDataLayerInstance>>& GetSelectedDataLayersFromEditorSelection() const;
+	bool SetParentDataLayerForDataLayersInternal(const TArray<UDataLayerInstance*>& DataLayers, UDataLayerInstance* ParentDataLayer);
 
 	bool UpdateAllActorsVisibility(const bool bNotifySelectionChange, const bool bRedrawViewports, ULevel* InLevel);
 
