Fixed Data Layer outliner not refreshed when starting PIE client with the option 'Run under one process' disabled.

diff --git a/Engine/Source/Editor/DataLayerEditor/Private/DataLayer/DataLayerMode.cpp b/Engine/Source/Editor/DataLayerEditor/Private/DataLayer/DataLayerMode.cpp
--- a/Engine/Source/Editor/DataLayerEditor/Private/DataLayer/DataLayerMode.cpp
+++ b/Engine/Source/Editor/DataLayerEditor/Private/DataLayer/DataLayerMode.cpp
@@ -136,6 +136,7 @@ FDataLayerMode::FDataLayerMode(const FDataLayerModeParams& Params)
 
 	USelection::SelectionChangedEvent.AddRaw(this, &FDataLayerMode::OnLevelSelectionChanged);
 	USelection::SelectObjectEvent.AddRaw(this, &FDataLayerMode::OnLevelSelectionChanged);
+	FCoreUObjectDelegates::PostLoadMapWithWorld.AddRaw(this, &FDataLayerMode::OnPostLoadMapWithWorld);
 
 	UWorldPartitionEditorPerProjectUserSettings* SharedSettings = GetMutableDefault<UWorldPartitionEditorPerProjectUserSettings>();
 	bHideEditorDataLayers = SharedSettings->bHideEditorDataLayers;
@@ -251,6 +252,7 @@ FDataLayerMode::~FDataLayerMode()
 {
 	USelection::SelectionChangedEvent.RemoveAll(this);
 	USelection::SelectObjectEvent.RemoveAll(this);
+	FCoreUObjectDelegates::PostLoadMapWithWorld.RemoveAll(this);
 }
 
 TSharedRef<FSceneOutlinerFilter> FDataLayerMode::CreateHideEditorDataLayersFilter()
@@ -2335,6 +2337,11 @@ void FDataLayerMode::SynchronizeSelection()
 	}
 }
 
+void FDataLayerMode::OnPostLoadMapWithWorld(UWorld* World)
+{
+	SceneOutliner->FullRefresh();
+}
+
 void FDataLayerMode::OnLevelSelectionChanged(UObject* Obj)
 {
 	if (!bShowOnlySelectedActors && !bHighlightSelectedDataLayers)
diff --git a/Engine/Source/Editor/DataLayerEditor/Private/DataLayer/DataLayerMode.h b/Engine/Source/Editor/DataLayerEditor/Private/DataLayer/DataLayerMode.h
--- a/Engine/Source/Editor/DataLayerEditor/Private/DataLayer/DataLayerMode.h
+++ b/Engine/Source/Editor/DataLayerEditor/Private/DataLayer/DataLayerMode.h
@@ -146,6 +146,7 @@ private:
 	TArray<UDataLayerInstance*> GetSelectedDataLayers(SSceneOutliner* InSceneOutliner) const;
 	void SetParentDataLayer(const TArray<UDataLayerInstance*> DataLayers, UDataLayerInstance* ParentDataLayer) const;
 	void OnLevelSelectionChanged(UObject* Obj);
+	void OnPostLoadMapWithWorld(UWorld* World);
 	static void CreateDataLayerPicker(UToolMenu* InMenu, FOnDataLayerInstancePicked OnDataLayerInstancePicked, FOnShouldFilterDataLayerInstance OnShouldFilterDataLayerInstance, bool bInShowRoot = false);
 	bool ShouldExpandDataLayer(const UDataLayerInstance* DataLayer) const;
 	bool ContainsSelectedChildDataLayer(const UDataLayerInstance* DataLayer) const;
