Allow to create new data layer asset when going through Data Layer Details panel.

diff --git a/Engine/Source/Editor/DataLayerEditor/Private/DataLayer/DataLayerInstanceCustomization.cpp b/Engine/Source/Editor/DataLayerEditor/Private/DataLayer/DataLayerInstanceCustomization.cpp
--- a/Engine/Source/Editor/DataLayerEditor/Private/DataLayer/DataLayerInstanceCustomization.cpp
+++ b/Engine/Source/Editor/DataLayerEditor/Private/DataLayer/DataLayerInstanceCustomization.cpp
@@ -64,6 +64,7 @@ void FDataLayerInstanceDetails::CustomizeDetails(IDetailLayoutBuilder& DetailBui
 			[
 				SNew(SObjectPropertyEntryBox)
 				.AllowClear(false)
+				.AllowCreate(true)
 				.AllowedClass(UDataLayerAsset::StaticClass())
 				.PropertyHandle(DataLayerAssetProperty)
 				.DisplayThumbnail(true)
