DataLayer Outliner: Prevent Delete Button from showing up for non User managed actors (ex: Partition Actors)

#rnx

diff --git a/Engine/Source/Editor/DataLayerEditor/Private/DataLayer/DataLayerOutlinerDeleteButtonColumn.cpp b/Engine/Source/Editor/DataLayerEditor/Private/DataLayer/DataLayerOutlinerDeleteButtonColumn.cpp
--- a/Engine/Source/Editor/DataLayerEditor/Private/DataLayer/DataLayerOutlinerDeleteButtonColumn.cpp
+++ b/Engine/Source/Editor/DataLayerEditor/Private/DataLayer/DataLayerOutlinerDeleteButtonColumn.cpp
@@ -62,7 +62,7 @@ const TSharedRef<SWidget> FDataLayerOutlinerDeleteButtonColumn::ConstructRowWidg
 			{
 				AActor* Actor = DataLayerActorItem->GetActor();
 				const UDataLayerInstance* DataLayerInstance = DataLayerActorItem->GetDataLayer();
-				return (Actor && DataLayerInstance && DataLayerInstance->CanRemoveActor(Actor) && DataLayerInstance->CanUserRemoveActors() && TreeItem->CanInteract()) ? EVisibility::Visible : EVisibility::Collapsed;
+				return (Actor && Actor->IsUserManaged() && DataLayerInstance && DataLayerInstance->CanRemoveActor(Actor) && DataLayerInstance->CanUserRemoveActors() && TreeItem->CanInteract()) ? EVisibility::Visible : EVisibility::Collapsed;
 			})
 			.OnClicked_Lambda([this, TreeItem, DataLayerActorItem]()
 			{
