Changed Pin/Unpin (user facing) to ForceLoad/ReleaseForceLoad in Data Layer Outliner
#tests Context menu in Data Layer Outliner

diff --git a/Engine/Source/Editor/DataLayerEditor/Private/DataLayer/DataLayerMode.cpp b/Engine/Source/Editor/DataLayerEditor/Private/DataLayer/DataLayerMode.cpp
--- a/Engine/Source/Editor/DataLayerEditor/Private/DataLayer/DataLayerMode.cpp
+++ b/Engine/Source/Editor/DataLayerEditor/Private/DataLayer/DataLayerMode.cpp
@@ -1750,23 +1750,23 @@ void FDataLayerMode::RegisterContextMenu()
 			{
 				FToolMenuSection& Section = InMenu->AddSection("DataLayerLoading", LOCTEXT("DataLayerLoading", "Loading"));
 
-				Section.AddMenuEntry("PinActorsInDataLayers", LOCTEXT("PinActorsInDataLayers", "Pin Actor(s) in Data Layer(s)"), FText(), FSlateIcon(),
+				Section.AddMenuEntry("ForceLoadActorsInDataLayers", LOCTEXT("ForceLoadActorsInDataLayers", "Force Load Actor(s) in Data Layer(s)"), FText(), FSlateIcon(),
 					FUIAction(
 						FExecuteAction::CreateLambda([SelectedDataLayers]() 
 						{
 							check(!SelectedDataLayers.IsEmpty());
-							const FScopedTransaction Transaction(LOCTEXT("PinActorsInDataLayers", "Pin Actor(s) in Data Layer(s)"));
+							const FScopedTransaction Transaction(LOCTEXT("ForceLoadActorsInDataLayers", "Force Load Actor(s) in Data Layer(s)"));
 							UDataLayerEditorSubsystem::Get()->SetActorsPinStateInDataLayers(SelectedDataLayers, true);
 						}),
 						FCanExecuteAction::CreateLambda([SelectedDataLayers] { return !SelectedDataLayers.IsEmpty(); })
 					));
 
-				Section.AddMenuEntry("UnpinActorsInDataLayers", LOCTEXT("UnpinActorsInDataLayers", "Unpin Actor(s) in Data Layer(s)"), FText(), FSlateIcon(),
+				Section.AddMenuEntry("ReleaseForceLoadActorsInDataLayers", LOCTEXT("ReleaseForceLoadActorsInDataLayers", "Release Force Load Actor(s) in Data Layer(s)"), FText(), FSlateIcon(),
 					FUIAction(
 						FExecuteAction::CreateLambda([SelectedDataLayers]()
 							{
 								check(!SelectedDataLayers.IsEmpty());
-								const FScopedTransaction Transaction(LOCTEXT("UnpinActorsInDataLayers", "Unpin Actor(s) in Data Layer(s)"));
+								const FScopedTransaction Transaction(LOCTEXT("ReleaseForceLoadActorsInDataLayers", "Release Force Load Actor(s) in Data Layer(s)"));
 								UDataLayerEditorSubsystem::Get()->SetActorsPinStateInDataLayers(SelectedDataLayers, false);
 							}),
 						FCanExecuteAction::CreateLambda([SelectedDataLayers] { return !SelectedDataLayers.IsEmpty(); })
