Don't allow to reparent a Data Layer part of an EDL hierarchy if the new parent is not part of the same EDL hierarchy.

diff --git a/Engine/Source/Editor/DataLayerEditor/Private/DataLayer/DataLayerMode.cpp b/Engine/Source/Editor/DataLayerEditor/Private/DataLayer/DataLayerMode.cpp
--- a/Engine/Source/Editor/DataLayerEditor/Private/DataLayer/DataLayerMode.cpp
+++ b/Engine/Source/Editor/DataLayerEditor/Private/DataLayer/DataLayerMode.cpp
@@ -1573,13 +1573,14 @@ void FDataLayerMode::RegisterContextMenu()
 					Section.AddSubMenu("MoveSelectedDataLayersTo", LOCTEXT("MoveSelectedDataLayersTo", "Move Selected Data Layer(s) To"), FText(),
 						FNewToolMenuDelegate::CreateLambda([Mode, SelectedDataLayers](UToolMenu* InSubMenu)
 						{
+							const bool bShowRoot = Algo::AllOf(SelectedDataLayers, [](const UDataLayerInstance* DataLayerInstance) { return DataLayerInstance->CanBeChildOf(nullptr); });
 							CreateDataLayerPicker(InSubMenu,
 								FOnDataLayerInstancePicked::CreateLambda([Mode, SelectedDataLayers](UDataLayerInstance* TargetDataLayerInstance)
 								{
 									TArray<UDataLayerInstance*> DataLayerInstances;
 									for (UDataLayerInstance* DataLayerInstance : SelectedDataLayers)
 									{
-										if (ensure(DataLayerInstance->CanBeChildOf(TargetDataLayerInstance)))
+										if (DataLayerInstance->CanBeChildOf(TargetDataLayerInstance))
 										{
 											DataLayerInstances.Add(DataLayerInstance);
 										}
@@ -1598,7 +1599,7 @@ void FDataLayerMode::RegisterContextMenu()
 									}
 									return false;
 								}),
-								/*bShowRoot*/true);
+								bShowRoot);
 						}));
 				}
 
diff --git a/Engine/Source/Runtime/Engine/Private/WorldPartition/DataLayer/DataLayerInstance.cpp b/Engine/Source/Runtime/Engine/Private/WorldPartition/DataLayer/DataLayerInstance.cpp
--- a/Engine/Source/Runtime/Engine/Private/WorldPartition/DataLayer/DataLayerInstance.cpp
+++ b/Engine/Source/Runtime/Engine/Private/WorldPartition/DataLayer/DataLayerInstance.cpp
@@ -409,6 +409,19 @@ bool UDataLayerInstance::CanBeChildOf(const UDataLayerInstance* InParent, FText*
 		return false;
 	}
 
+	// Check if root External Data Layer is the same for the parent
+	if (const UExternalDataLayerInstance* ExternalDataLayerInstance = GetRootExternalDataLayerInstance())
+	{
+		if (!InParent || (InParent->GetRootExternalDataLayerInstance() != ExternalDataLayerInstance))
+		{
+			if (OutReason)
+			{
+				*OutReason = LOCTEXT("ParentHasDifferentRootExternalDataLayer", "Can't reparent Data Layer in a different root External Data Layer");
+			}
+			return false;
+		}
+	}
+
 	if (InParent == nullptr)
 	{
 		// nullptr is considered a valid parent
