When replacing an actor, if the new actor is not in an EDL, make sure to move the actor into a new package and delete the old one.
Also, try to match the original actor's EDL when replacing an actor.

diff --git a/Engine/Source/Editor/DataLayerEditor/Private/DataLayer/DataLayerEditorSubsystem.cpp b/Engine/Source/Editor/DataLayerEditor/Private/DataLayer/DataLayerEditorSubsystem.cpp
--- a/Engine/Source/Editor/DataLayerEditor/Private/DataLayer/DataLayerEditorSubsystem.cpp
+++ b/Engine/Source/Editor/DataLayerEditor/Private/DataLayer/DataLayerEditorSubsystem.cpp
@@ -436,15 +436,30 @@ const UExternalDataLayerInstance* UDataLayerEditorSubsystem::GetActorSpawningExt
 
 void UDataLayerEditorSubsystem::OnActorPreSpawnInitialization(AActor* InActor)
 {
-	if (const UExternalDataLayerInstance* ExternalDataLayerInstance = GetActorSpawningExternalDataLayerInstance(InActor))
+	auto MoveActorToExternalDataLayer = [this](AActor* Actor, const UExternalDataLayerInstance* ExternalDataLayerInstance)
 	{
-		check(!InActor->GetExternalDataLayerAsset());
 		FText FailureReason;
-		if (!FExternalDataLayerHelper::MoveActorsToExternalDataLayer({ InActor }, ExternalDataLayerInstance, &FailureReason))
+		if (!FExternalDataLayerHelper::MoveActorsToExternalDataLayer({ Actor }, ExternalDataLayerInstance, &FailureReason))
 		{
 			UE_LOG(LogWorldPartition, Warning, TEXT("%s"), *FailureReason.ToString());
 			LastWarningNotification = FailureReason;
 		}
+	};
+
+	if (const UExternalDataLayerInstance* ExternalDataLayerInstance = GetActorSpawningExternalDataLayerInstance(InActor))
+	{
+		check(!InActor->GetExternalDataLayerAsset());
+		MoveActorToExternalDataLayer(InActor, ExternalDataLayerInstance);
+	}
+	
+	if (!InActor->GetExternalDataLayerAsset())
+	{
+		// Here we repair the case where the actor package is an EDL package but the actor is not in an EDL.
+		// This can happen when using 'replace actor' and the class or context doesn't target an EDL, but the replace logic always tried to reuse the old actor package.
+		if (UPackage* ActorPackage = InActor->GetExternalPackage(); ActorPackage && FExternalDataLayerHelper::IsExternalDataLayerPath(ActorPackage->GetPathName()))
+		{
+			MoveActorToExternalDataLayer(InActor, nullptr);
+		}
 	}
 }
 
@@ -452,7 +467,7 @@ void UDataLayerEditorSubsystem::OnEditorActorReplaced(AActor* InOldActor, AActor
 {
 	// Try to apply the current context on the new replacing actor
 	const bool bForceTryApply = true;
-	ApplyContext(InNewActor, bForceTryApply);
+	ApplyContext(InNewActor, bForceTryApply, InOldActor);
 }
 
 void UDataLayerEditorSubsystem::OnNewActorsPlaced(UObject* InObjToUse, const TArray<AActor*>& InPlacedActors)
@@ -523,7 +538,7 @@ void UDataLayerEditorSubsystem::BeginDestroy()
 	Super::BeginDestroy();
 }
 
-void UDataLayerEditorSubsystem::ApplyContext(AActor* InActor, bool bInForceTryApply)
+void UDataLayerEditorSubsystem::ApplyContext(AActor* InActor, bool bInForceTryApply, AActor* InReplacedActor)
 {
 	if (GIsReinstancing || InActor->bIsEditorPreviewActor || !InActor->IsPackageExternal() || !InActor->IsMainPackageActor() || (InActor->GetWorld() != GetWorld()))
 	{
@@ -539,24 +554,38 @@ void UDataLayerEditorSubsystem::ApplyContext(AActor* InActor, bool bInForceTryAp
 	// Try to apply context External Data Layer (this operation can fail if asset referencing validation fails)
 	const UExternalDataLayerAsset* CurrentExternalDataLayer = GetActorEditorContextCurrentExternalDataLayer();
 	const UExternalDataLayerAsset* ActorExternalDataLayerAsset = InActor->GetExternalDataLayerAsset();
-	if (CurrentExternalDataLayer && (ActorExternalDataLayerAsset != CurrentExternalDataLayer))
+	const UExternalDataLayerAsset* ReplacedActorExternalDataLayerAsset = InReplacedActor ? InReplacedActor->GetExternalDataLayerAsset() : nullptr;
+	if (CurrentExternalDataLayer)
 	{
-		// Don't apply if there's a valid override spawning External Data Layer (see OnActorPreSpawnInitialization)
-		const UExternalDataLayerInstance* SpawningExternalDataLayer = ULevel::GetOverrideSpawningLevelMountPointObject() ? GetActorSpawningExternalDataLayerInstance(InActor) : nullptr;
-		if (!SpawningExternalDataLayer || bInForceTryApply)
+		if (ActorExternalDataLayerAsset != CurrentExternalDataLayer)
 		{
-			UExternalDataLayerManager* ExternalDataLayerManager = UExternalDataLayerManager::GetExternalDataLayerManager(InActor);
-			if (const UExternalDataLayerInstance* CurrentExternalDataLayerInstance = CurrentExternalDataLayer ? ExternalDataLayerManager->GetExternalDataLayerInstance(CurrentExternalDataLayer) : nullptr)
+			// Don't apply if there's a valid override spawning External Data Layer (see OnActorPreSpawnInitialization)
+			const UExternalDataLayerInstance* SpawningExternalDataLayer = ULevel::GetOverrideSpawningLevelMountPointObject() ? GetActorSpawningExternalDataLayerInstance(InActor) : nullptr;
+			if (!SpawningExternalDataLayer || bInForceTryApply)
 			{
-				FText FailureReason;
-				if (!FExternalDataLayerHelper::MoveActorsToExternalDataLayer({ InActor }, CurrentExternalDataLayerInstance, &FailureReason))
+				UExternalDataLayerManager* ExternalDataLayerManager = UExternalDataLayerManager::GetExternalDataLayerManager(InActor);
+				if (const UExternalDataLayerInstance* CurrentExternalDataLayerInstance = CurrentExternalDataLayer ? ExternalDataLayerManager->GetExternalDataLayerInstance(CurrentExternalDataLayer) : nullptr)
 				{
-					UE_LOG(LogWorldPartition, Warning, TEXT("%s"), *FailureReason.ToString());
-					LastWarningNotification = FailureReason;
+					FText FailureReason;
+					if (!FExternalDataLayerHelper::MoveActorsToExternalDataLayer({ InActor }, CurrentExternalDataLayerInstance, &FailureReason))
+					{
+						UE_LOG(LogWorldPartition, Warning, TEXT("%s"), *FailureReason.ToString());
+						LastWarningNotification = FailureReason;
+					}
 				}
 			}
 		}
 	}
+	// When replacing an actor, try to match the replaced actor External Data Layer if there's none set in the Actor Editor Context
+	else if (ReplacedActorExternalDataLayerAsset && (ActorExternalDataLayerAsset != ReplacedActorExternalDataLayerAsset))
+	{
+		UExternalDataLayerManager* ExternalDataLayerManager = UExternalDataLayerManager::GetExternalDataLayerManager(InActor);
+		if (const UExternalDataLayerInstance* ReplacedExternalDataLayerInstance = ExternalDataLayerManager ? ExternalDataLayerManager->GetExternalDataLayerInstance(ReplacedActorExternalDataLayerAsset) : nullptr)
+		{
+			// As this is an attempt, no need to report failures
+			FExternalDataLayerHelper::MoveActorsToExternalDataLayer({ InActor }, ReplacedExternalDataLayerInstance);
+		}
+	}
 
 	// Apply context Data Layers (except External Data Layer)
 	TArray<UDataLayerInstance*> DataLayerInstances = DataLayerManager->GetActorEditorContextDataLayers();
diff --git a/Engine/Source/Editor/DataLayerEditor/Public/DataLayer/DataLayerEditorSubsystem.h b/Engine/Source/Editor/DataLayerEditor/Public/DataLayer/DataLayerEditorSubsystem.h
--- a/Engine/Source/Editor/DataLayerEditor/Public/DataLayer/DataLayerEditorSubsystem.h
+++ b/Engine/Source/Editor/DataLayerEditor/Public/DataLayer/DataLayerEditorSubsystem.h
@@ -776,7 +776,7 @@ private:
 	void OnEditorActorReplaced(AActor* InOldActor, AActor* InNewActor);
 	void OnExternalDataLayerAssetRegistrationStateChanged(const UExternalDataLayerAsset* ExternalDataLayerAsset, EExternalDataLayerRegistrationState OldState, EExternalDataLayerRegistrationState NewState);
 	TUniquePtr<FLevelEditorDragDropWorldSurrogateReferencingObject> OnLevelEditorDragDropWorldSurrogateReferencingObject(UWorld* ReferencingWorld, const FSoftObjectPath& Object);
-	void ApplyContext(AActor* InActor, bool bInForceTryApply = false);
+	void ApplyContext(AActor* InActor, bool bInForceTryApply = false, AActor* InReplacedActor = nullptr);
 	const UExternalDataLayerInstance* GetActorSpawningExternalDataLayerInstance(AActor* InActor) const;
 	static const UExternalDataLayerAsset* GetReferencingWorldSurrogateObjectForObject(UWorld* ReferencingWorld, const FSoftObjectPath& ObjectPath);
 
diff --git a/Engine/Source/Runtime/Engine/Private/WorldPartition/DataLayer/ExternalDataLayerHelper.cpp b/Engine/Source/Runtime/Engine/Private/WorldPartition/DataLayer/ExternalDataLayerHelper.cpp
--- a/Engine/Source/Runtime/Engine/Private/WorldPartition/DataLayer/ExternalDataLayerHelper.cpp
+++ b/Engine/Source/Runtime/Engine/Private/WorldPartition/DataLayer/ExternalDataLayerHelper.cpp
@@ -141,8 +141,6 @@ bool FExternalDataLayerHelper::CanMoveActorsToExternalDataLayer(const TArray<AAc
 		check(!InActor->IsTemplate());
 		check(InActor->GetLevel());
 
-		const UExternalDataLayerAsset* OldExternalDataLayerAsset = InActor->GetExternalDataLayerAsset();
-		const UExternalDataLayerAsset* NewExternalDataLayerAsset = InExternalDataLayerInstance ? InExternalDataLayerInstance->GetExternalDataLayerAsset() : nullptr;
 		if (!InActor->IsPackageExternal())
 		{
 			OutFailureReason = FText::Format(LOCTEXT("CantMoveActorToEDL_NotPackageExternal", "Actor {0} is not using external package."), FText::FromString(InActor->GetName()));
@@ -161,16 +159,29 @@ bool FExternalDataLayerHelper::CanMoveActorsToExternalDataLayer(const TArray<AAc
 			return false;
 		}
 
-		if (!OldExternalDataLayerAsset && !NewExternalDataLayerAsset)
-		{
-			OutFailureReason = FText::Format(LOCTEXT("CantMoveActorToEDL_NoExternalDataLayer", "Actor {0} has already no External Data Layer."), FText::FromString(InActor->GetName()));
-			return false;
-		}
+		FExternalDataLayerUID OldExternalDataLayerUIDFromPackage;
+		const UPackage* OldPackage = InActor->GetExternalPackage();
+		FExternalDataLayerHelper::IsExternalDataLayerPath(OldPackage->GetPathName(), &OldExternalDataLayerUIDFromPackage);
+		const UExternalDataLayerAsset* OldExternalDataLayerAsset = InActor->GetExternalDataLayerAsset();
+		const FExternalDataLayerUID OldExternalDataLayerUIDFromAsset = OldExternalDataLayerAsset ? OldExternalDataLayerAsset->GetUID() : FExternalDataLayerUID();
+		const UExternalDataLayerAsset* NewExternalDataLayerAsset = InExternalDataLayerInstance ? InExternalDataLayerInstance->GetExternalDataLayerAsset() : nullptr;
 
-		if (OldExternalDataLayerAsset == NewExternalDataLayerAsset)
+		// Detect if old actor package and EDL asset mismatch (this can happen during the replace actor process).
+		// In this case, skip the optimization that tries to detect no change that are relying on OldExternalDataLayerAsset (as it's unreliable).
+		const bool bDoesOldExternalDataLayerUIDMismatch = OldExternalDataLayerUIDFromAsset != OldExternalDataLayerUIDFromPackage;
+		if (!bDoesOldExternalDataLayerUIDMismatch)
 		{
-			OutFailureReason = FText::Format(LOCTEXT("CantMoveActorToEDL_SameExternalDataLayer", "Actor {0} is already assigned to this External Data Layer."), FText::FromString(InActor->GetName()));
-			return false;
+			if (!OldExternalDataLayerAsset && !NewExternalDataLayerAsset)
+			{
+				OutFailureReason = FText::Format(LOCTEXT("CantMoveActorToEDL_NoExternalDataLayer", "Actor {0} has already no External Data Layer."), FText::FromString(InActor->GetName()));
+				return false;
+			}
+
+			if (OldExternalDataLayerAsset == NewExternalDataLayerAsset)
+			{
+				OutFailureReason = FText::Format(LOCTEXT("CantMoveActorToEDL_SameExternalDataLayer", "Actor {0} is already assigned to this External Data Layer."), FText::FromString(InActor->GetName()));
+				return false;
+			}
 		}
 
 		if (NewExternalDataLayerAsset && !InActor->SupportsDataLayerType(UExternalDataLayerInstance::StaticClass()))
