Moved External Data Layer logic that repairs a mismatch between the actor package and the actor's EDL asset in UDataLayerEditorSubsystem::OnEditorActorReplaced.
Also fixed the code to only execute when appropriate.

diff --git a/Engine/Source/Editor/DataLayerEditor/Private/DataLayer/DataLayerEditorSubsystem.cpp b/Engine/Source/Editor/DataLayerEditor/Private/DataLayer/DataLayerEditorSubsystem.cpp
--- a/Engine/Source/Editor/DataLayerEditor/Private/DataLayer/DataLayerEditorSubsystem.cpp
+++ b/Engine/Source/Editor/DataLayerEditor/Private/DataLayer/DataLayerEditorSubsystem.cpp
@@ -390,15 +390,18 @@ void UDataLayerEditorSubsystem::Deinitialize()
 	UEditorStateSubsystem::UnregisterEditorStateType<UDataLayerEditorState>();
 }
 
-const UExternalDataLayerInstance* UDataLayerEditorSubsystem::GetActorSpawningExternalDataLayerInstance(AActor* InActor) const
+bool UDataLayerEditorSubsystem::ShouldHandleActor(AActor* InActor) const
 {
-	if (GIsReinstancing || InActor->bIsEditorPreviewActor || !InActor->IsPackageExternal() || !InActor->IsMainPackageActor())
+	if (GIsReinstancing || InActor->bIsEditorPreviewActor || !InActor->IsPackageExternal() || !InActor->IsMainPackageActor() || !InActor->GetWorld() || (InActor->GetWorld() != GetWorld()))
 	{
-		return nullptr;
+		return false;
 	}
-		
-	const UWorld* OwningWorld = InActor->GetWorld();
-	if (!OwningWorld || (OwningWorld != GetWorld()))
+	return true;
+}
+
+const UExternalDataLayerInstance* UDataLayerEditorSubsystem::GetActorSpawningExternalDataLayerInstance(AActor* InActor) const
+{
+	if (!ShouldHandleActor(InActor))
 	{
 		return nullptr;
 	}
@@ -409,6 +412,7 @@ const UExternalDataLayerInstance* UDataLayerEditorSubsystem::GetActorSpawningExt
 		return nullptr;
 	}
 
+	const UWorld* OwningWorld = InActor->GetWorld();
 	const UObject* OverrideSpawningLevelMountPointObject = ULevel::GetOverrideSpawningLevelMountPointObject();
 	const UExternalDataLayerManager* ExternalDataLayerManager = UExternalDataLayerManager::GetExternalDataLayerManager(OwningWorld);
 	
@@ -434,33 +438,26 @@ const UExternalDataLayerInstance* UDataLayerEditorSubsystem::GetActorSpawningExt
 	return ResolvedExternalDataLayerInstance;
 }
 
-void UDataLayerEditorSubsystem::OnActorPreSpawnInitialization(AActor* InActor)
+void UDataLayerEditorSubsystem::MoveActorToExternalDataLayer(AActor* InActor, const UExternalDataLayerInstance* InExternalDataLayerInstance, bool bInNotifyFailureReason)
 {
-	auto MoveActorToExternalDataLayer = [this](AActor* Actor, const UExternalDataLayerInstance* ExternalDataLayerInstance)
+	FText FailureReason;
+	if (!FExternalDataLayerHelper::MoveActorsToExternalDataLayer({ InActor }, InExternalDataLayerInstance, &FailureReason))
 	{
-		FText FailureReason;
-		if (!FExternalDataLayerHelper::MoveActorsToExternalDataLayer({ Actor }, ExternalDataLayerInstance, &FailureReason))
+		if (bInNotifyFailureReason)
 		{
 			UE_LOG(LogWorldPartition, Warning, TEXT("%s"), *FailureReason.ToString());
 			LastWarningNotification = FailureReason;
 		}
-	};
+	}
+};
 
+void UDataLayerEditorSubsystem::OnActorPreSpawnInitialization(AActor* InActor)
+{
 	if (const UExternalDataLayerInstance* ExternalDataLayerInstance = GetActorSpawningExternalDataLayerInstance(InActor))
 	{
 		check(!InActor->GetExternalDataLayerAsset());
 		MoveActorToExternalDataLayer(InActor, ExternalDataLayerInstance);
 	}
-	
-	if (!InActor->GetExternalDataLayerAsset())
-	{
-		// Here we repair the case where the actor package is an EDL package but the actor is not in an EDL.
-		// This can happen when using 'replace actor' and the class or context doesn't target an EDL, but the replace logic always tried to reuse the old actor package.
-		if (UPackage* ActorPackage = InActor->GetExternalPackage(); ActorPackage && FExternalDataLayerHelper::IsExternalDataLayerPath(ActorPackage->GetPathName()))
-		{
-			MoveActorToExternalDataLayer(InActor, nullptr);
-		}
-	}
 }
 
 void UDataLayerEditorSubsystem::OnEditorActorReplaced(AActor* InOldActor, AActor* InNewActor)
@@ -468,6 +465,24 @@ void UDataLayerEditorSubsystem::OnEditorActorReplaced(AActor* InOldActor, AActor
 	// Try to apply the current context on the new replacing actor
 	const bool bForceTryApply = true;
 	ApplyContext(InNewActor, bForceTryApply, InOldActor);
+
+	if (ShouldHandleActor(InNewActor))
+	{
+		// Here we repair the case where the actor package doesn't match the actor EDL asset.
+		// This can happen when using 'replace actor' as it reuses the old actor's package.
+		// In this case, choose the EDL from the package if any.
+		UPackage* ActorPackage = InNewActor->GetExternalPackage();
+		const UWorld* OwningWorld = InNewActor->GetWorld();
+		if (const UExternalDataLayerManager* ExternalDataLayerManager = UExternalDataLayerManager::GetExternalDataLayerManager(OwningWorld); ActorPackage && ExternalDataLayerManager)
+		{
+			const UExternalDataLayerAsset* PackageDataLayerAsset = ExternalDataLayerManager->GetMatchingExternalDataLayerAssetForObjectPath(ActorPackage->GetPathName());
+			if (PackageDataLayerAsset != InNewActor->GetExternalDataLayerAsset())
+			{
+				const UExternalDataLayerInstance* ExternalDataLayerInstance = ExternalDataLayerManager->GetExternalDataLayerInstance(PackageDataLayerAsset);
+				MoveActorToExternalDataLayer(InNewActor, ExternalDataLayerInstance);
+			}
+		}
+	}
 }
 
 void UDataLayerEditorSubsystem::OnNewActorsPlaced(UObject* InObjToUse, const TArray<AActor*>& InPlacedActors)
@@ -540,7 +555,7 @@ void UDataLayerEditorSubsystem::BeginDestroy()
 
 void UDataLayerEditorSubsystem::ApplyContext(AActor* InActor, bool bInForceTryApply, AActor* InReplacedActor)
 {
-	if (GIsReinstancing || InActor->bIsEditorPreviewActor || !InActor->IsPackageExternal() || !InActor->IsMainPackageActor() || (InActor->GetWorld() != GetWorld()))
+	if (!ShouldHandleActor(InActor))
 	{
 		return;
 	}
@@ -566,12 +581,7 @@ void UDataLayerEditorSubsystem::ApplyContext(AActor* InActor, bool bInForceTryAp
 				UExternalDataLayerManager* ExternalDataLayerManager = UExternalDataLayerManager::GetExternalDataLayerManager(InActor);
 				if (const UExternalDataLayerInstance* CurrentExternalDataLayerInstance = CurrentExternalDataLayer ? ExternalDataLayerManager->GetExternalDataLayerInstance(CurrentExternalDataLayer) : nullptr)
 				{
-					FText FailureReason;
-					if (!FExternalDataLayerHelper::MoveActorsToExternalDataLayer({ InActor }, CurrentExternalDataLayerInstance, &FailureReason))
-					{
-						UE_LOG(LogWorldPartition, Warning, TEXT("%s"), *FailureReason.ToString());
-						LastWarningNotification = FailureReason;
-					}
+					MoveActorToExternalDataLayer(InActor, CurrentExternalDataLayerInstance);
 				}
 			}
 		}
@@ -583,7 +593,8 @@ void UDataLayerEditorSubsystem::ApplyContext(AActor* InActor, bool bInForceTryAp
 		if (const UExternalDataLayerInstance* ReplacedExternalDataLayerInstance = ExternalDataLayerManager ? ExternalDataLayerManager->GetExternalDataLayerInstance(ReplacedActorExternalDataLayerAsset) : nullptr)
 		{
 			// As this is an attempt, no need to report failures
-			FExternalDataLayerHelper::MoveActorsToExternalDataLayer({ InActor }, ReplacedExternalDataLayerInstance);
+			const bool bNotifyFailureReason = false;
+			MoveActorToExternalDataLayer(InActor, ReplacedExternalDataLayerInstance, bNotifyFailureReason);
 		}
 	}
 
diff --git a/Engine/Source/Editor/DataLayerEditor/Public/DataLayer/DataLayerEditorSubsystem.h b/Engine/Source/Editor/DataLayerEditor/Public/DataLayer/DataLayerEditorSubsystem.h
--- a/Engine/Source/Editor/DataLayerEditor/Public/DataLayer/DataLayerEditorSubsystem.h
+++ b/Engine/Source/Editor/DataLayerEditor/Public/DataLayer/DataLayerEditorSubsystem.h
@@ -779,6 +779,8 @@ private:
 	void ApplyContext(AActor* InActor, bool bInForceTryApply = false, AActor* InReplacedActor = nullptr);
 	const UExternalDataLayerInstance* GetActorSpawningExternalDataLayerInstance(AActor* InActor) const;
 	static const UExternalDataLayerAsset* GetReferencingWorldSurrogateObjectForObject(UWorld* ReferencingWorld, const FSoftObjectPath& ObjectPath);
+	void MoveActorToExternalDataLayer(AActor* InActor, const UExternalDataLayerInstance* InExternalDataLayerInstance, bool bInNotifyFailureReason = true);
+	bool ShouldHandleActor(AActor* InActor) const;
 
 	/** Contains Data Layers that contain actors that are part of the editor selection */
 	mutable TSet<TWeakObjectPtr<const UDataLayerInstance>> SelectedDataLayersFromEditorSelection;
diff --git a/Engine/Source/Runtime/Engine/Private/WorldPartition/DataLayer/ExternalDataLayerHelper.cpp b/Engine/Source/Runtime/Engine/Private/WorldPartition/DataLayer/ExternalDataLayerHelper.cpp
--- a/Engine/Source/Runtime/Engine/Private/WorldPartition/DataLayer/ExternalDataLayerHelper.cpp
+++ b/Engine/Source/Runtime/Engine/Private/WorldPartition/DataLayer/ExternalDataLayerHelper.cpp
@@ -257,8 +257,7 @@ bool FExternalDataLayerHelper::CanMoveActorsToExternalDataLayer(const TArray<AAc
 			}
 			else
 			{
-				check(OldExternalDataLayerAsset);
-				OutFailureReason = FText::Format(LOCTEXT("CantRemoveEDLFromActorReferenceRestrictions", "Can't remove External Data Layer {0} from Actor {1}.\n{2}"), FText::FromString(OldExternalDataLayerAsset->GetName()), FText::FromString(InActor->GetName()), FText::FromString(JoinedReasons));
+				OutFailureReason = FText::Format(LOCTEXT("CantRemoveEDLFromActorReferenceRestrictions", "Can't remove External Data Layer {0} from Actor {1}.\n{2}"), FText::FromString(OldExternalDataLayerAsset ? OldExternalDataLayerAsset->GetName() : OldExternalDataLayerUIDFromPackage.ToString()), FText::FromString(InActor->GetName()), FText::FromString(JoinedReasons));
 			}
 			return false;
 		}
