Added more logging to better report when the creation of an External Data Layer instance fails.

diff --git a/Engine/Source/Editor/DataLayerEditor/Private/DataLayer/DataLayerEditorSubsystem.cpp b/Engine/Source/Editor/DataLayerEditor/Private/DataLayer/DataLayerEditorSubsystem.cpp
--- a/Engine/Source/Editor/DataLayerEditor/Private/DataLayer/DataLayerEditorSubsystem.cpp
+++ b/Engine/Source/Editor/DataLayerEditor/Private/DataLayer/DataLayerEditorSubsystem.cpp
@@ -1727,24 +1727,39 @@ UDataLayerInstance* UDataLayerEditorSubsystem::CreateDataLayerInstance(const FDa
 				if (UExternalDataLayerAsset* ExternalDataLayerAsset = Cast<UExternalDataLayerAsset>(Parameters.DataLayerAsset))
 				{
 					UWorld* OuterWorld = WorldDataLayers->GetTypedOuter<UWorld>();
-					UExternalDataLayerManager* ExternalDataLayerManager = UExternalDataLayerManager::GetExternalDataLayerManager(OuterWorld);
-					if (ExternalDataLayerManager && ExternalDataLayerManager->CanInjectExternalDataLayerAsset(ExternalDataLayerAsset))
+					if (UExternalDataLayerManager* ExternalDataLayerManager = UExternalDataLayerManager::GetExternalDataLayerManager(OuterWorld))
 					{
-						const bool bAllowCreate = true;
-						if (AWorldDataLayers* EDLWorldDataLayers = ExternalDataLayerManager->GetWorldDataLayers(ExternalDataLayerAsset, bAllowCreate))
+						FText FailureReason;
+						if (ExternalDataLayerManager->CanInjectExternalDataLayerAsset(ExternalDataLayerAsset, &FailureReason))
 						{
-							UExternalDataLayerInstance* ExternalDataLayerInstance = EDLWorldDataLayers->GetExternalDataLayerInstance(ExternalDataLayerAsset);
-							if (ensure(!ExternalDataLayerInstance))
+							const bool bAllowCreate = true;
+							if (AWorldDataLayers* EDLWorldDataLayers = ExternalDataLayerManager->GetWorldDataLayers(ExternalDataLayerAsset, bAllowCreate))
 							{
-								// Create External Data Layer Instance
-								ExternalDataLayerInstance = EDLWorldDataLayers->CreateDataLayer<UExternalDataLayerInstance>(ExternalDataLayerAsset);
-								if (ExternalDataLayerInstance && ExternalDataLayerManager->InjectExternalDataLayer(ExternalDataLayerAsset))
+								UExternalDataLayerInstance* ExternalDataLayerInstance = EDLWorldDataLayers->GetExternalDataLayerInstance(ExternalDataLayerAsset);
+								if (ensure(!ExternalDataLayerInstance))
 								{
-									NewDataLayerInstance = const_cast<UExternalDataLayerInstance*>(ExternalDataLayerManager->GetExternalDataLayerInstance(ExternalDataLayerAsset));
-									ensure(NewDataLayerInstance == ExternalDataLayerInstance);
+									// Create External Data Layer Instance
+									ExternalDataLayerInstance = EDLWorldDataLayers->CreateDataLayer<UExternalDataLayerInstance>(ExternalDataLayerAsset);
+									if (ensure(ExternalDataLayerInstance))
+									{
+										if (ExternalDataLayerManager->InjectExternalDataLayer(ExternalDataLayerAsset))
+										{
+											NewDataLayerInstance = const_cast<UExternalDataLayerInstance*>(ExternalDataLayerManager->GetExternalDataLayerInstance(ExternalDataLayerAsset));
+											ensure(NewDataLayerInstance == ExternalDataLayerInstance);
+										}
+										else
+										{
+											UE_LOG(LogDataLayerEditorSubsystem, Warning, TEXT("[EDL %s] Failed to inject External Data Layer."), *ExternalDataLayerAsset->GetName());
+										}
+									}
 								}
 							}
 						}
+						else
+						{
+							UE_LOG(LogDataLayerEditorSubsystem, Warning, TEXT("%s"), *FailureReason.ToString());
+							LastWarningNotification = FailureReason;
+						}
 					}
 					UE_CLOG(!NewDataLayerInstance, LogDataLayerEditorSubsystem, Error, TEXT("[EDL %s] Failed to create External Data Layer Instance."), *ExternalDataLayerAsset->GetName());
 				}
