ActorEditorContext - Added state objects to better support context serialization
* Capture/Restore actor editor context with  UActorEditorContextSubsystem::Capture/RestoreContext()
* Each client is responsible for handling it's own state and storing it in UActorEditorContextStateCollection

diff --git a/Engine/Source/Editor/DataLayerEditor/Private/DataLayer/DataLayerEditorSubsystem.cpp b/Engine/Source/Editor/DataLayerEditor/Private/DataLayer/DataLayerEditorSubsystem.cpp
--- a/Engine/Source/Editor/DataLayerEditor/Private/DataLayer/DataLayerEditorSubsystem.cpp
+++ b/Engine/Source/Editor/DataLayerEditor/Private/DataLayer/DataLayerEditorSubsystem.cpp
@@ -649,7 +649,79 @@ void UDataLayerEditorSubsystem::OnExecuteActorEditorContextAction(UWorld* InWorl
 			break;
 		}
 	}
-}	
+}
+
+void UDataLayerEditorSubsystem::CaptureActorEditorContextState(UWorld* InWorld, UActorEditorContextStateCollection* InStateCollection) const
+{
+	UActorEditorContextDataLayerState* State = nullptr;
+
+	if (UDataLayerManager* DataLayerManager = UDataLayerManager::GetDataLayerManager(InWorld))
+	{
+		for (UDataLayerInstance* DataLayerInstance : DataLayerManager->GetActorEditorContextDataLayers())
+		{
+			if (const UDataLayerAsset* DataLayerAsset = DataLayerInstance->GetAsset())
+			{
+				if (State == nullptr)
+				{
+					State = NewObject<UActorEditorContextDataLayerState>(InStateCollection);
+				}
+
+				if (const UExternalDataLayerAsset* ExternalDataLayerAsset = Cast<UExternalDataLayerAsset>(DataLayerAsset))
+				{
+					State->ExternalDataLayerAsset = ExternalDataLayerAsset;
+				}
+				else
+				{
+					State->DataLayerAssets.Add(DataLayerAsset);
+				}
+			}
+		}
+	}
+
+	if (State)
+	{
+		InStateCollection->AddState(State);
+	}
+}
+
+void UDataLayerEditorSubsystem::RestoreActorEditorContextState(UWorld* InWorld, const UActorEditorContextStateCollection* InStateCollection)
+{
+	if (UDataLayerManager* DataLayerManager = UDataLayerManager::GetDataLayerManager(InWorld))
+	{
+		TSet<UDataLayerInstance*> CurrentDataLayers(DataLayerManager->GetActorEditorContextDataLayers());
+		TSet<const UDataLayerInstance*> DataLayersToRestore;
+
+		if (const UActorEditorContextDataLayerState* State = InStateCollection->GetState<UActorEditorContextDataLayerState>())
+		{
+			for (const TSoftObjectPtr<const UDataLayerAsset>& DLAsset : State->DataLayerAssets)
+			{
+				const UDataLayerInstance* DataLayerInstance = DataLayerManager->GetDataLayerInstanceFromAssetName(FName(DLAsset.ToString()));
+				if (DataLayerInstance && DataLayerInstance->CanBeInActorEditorContext())
+				{
+					DataLayersToRestore.Add(DataLayerInstance);
+				}
+			}
+		}
+
+		// Add DL instances not in current context
+		for (const UDataLayerInstance* DataLayerInstance : DataLayersToRestore)
+		{
+			if (!CurrentDataLayers.Contains(DataLayerInstance))
+			{
+				AddToActorEditorContext(const_cast<UDataLayerInstance*>(DataLayerInstance));
+			}
+		}
+
+		// Remove DL instances currently in the context but not in the state we are restoring
+		for (UDataLayerInstance* DataLayerInstance : CurrentDataLayers)
+		{
+			if (!DataLayersToRestore.Contains(DataLayerInstance))
+			{
+				RemoveFromActorEditorContext(DataLayerInstance);
+			}
+		}
+	}
+}
 
 bool UDataLayerEditorSubsystem::GetActorEditorContextDisplayInfo(UWorld* InWorld, FActorEditorContextClientDisplayInfo& OutDiplayInfo) const
 {
diff --git a/Engine/Source/Editor/DataLayerEditor/Public/DataLayer/DataLayerEditorSubsystem.h b/Engine/Source/Editor/DataLayerEditor/Public/DataLayer/DataLayerEditorSubsystem.h
--- a/Engine/Source/Editor/DataLayerEditor/Public/DataLayer/DataLayerEditorSubsystem.h
+++ b/Engine/Source/Editor/DataLayerEditor/Public/DataLayer/DataLayerEditorSubsystem.h
@@ -2,6 +2,7 @@
 
 #pragma once
 
+#include "ActorEditorContextState.h"
 #include "Containers/Array.h"
 #include "Containers/Set.h"
 #include "CoreMinimal.h"
@@ -69,6 +70,19 @@ struct DATALAYEREDITOR_API FDataLayerCreationParameters
 	bool bIsPrivate;
 };
 
+UCLASS()
+class UActorEditorContextDataLayerState : public UActorEditorContextClientState
+{
+	GENERATED_BODY()
+
+public:
+	UPROPERTY(VisibleAnywhere, Category="Data Layers")
+	TObjectPtr<const UExternalDataLayerAsset> ExternalDataLayerAsset;
+	
+	UPROPERTY(VisibleAnywhere, Category="Data Layers")
+	TArray<TSoftObjectPtr<const UDataLayerAsset>> DataLayerAssets;
+};
+
 UCLASS()
 class DATALAYEREDITOR_API UDataLayerEditorSubsystem final : public UEditorSubsystem, public IActorEditorContextClient, public FTickableGameObject
 {
@@ -111,6 +125,8 @@ public:
 
 	//~ Begin IActorEditorContextClient interface
 	virtual void OnExecuteActorEditorContextAction(UWorld* InWorld, const EActorEditorContextAction& InType, AActor* InActor = nullptr) override;
+	virtual void CaptureActorEditorContextState(UWorld* InWorld, UActorEditorContextStateCollection* InStateCollection) const override;
+	virtual void RestoreActorEditorContextState(UWorld* InWorld, const UActorEditorContextStateCollection* InStateCollection) override;
 	virtual bool GetActorEditorContextDisplayInfo(UWorld* InWorld, FActorEditorContextClientDisplayInfo& OutDiplayInfo) const override;
 	virtual bool CanResetContext(UWorld* InWorld) const override { return true; };
 	virtual TSharedRef<SWidget> GetActorEditorContextWidget(UWorld* InWorld) const override;
diff --git a/Engine/Source/Editor/LevelEditor/Public/LevelEditorSubsystem.h b/Engine/Source/Editor/LevelEditor/Public/LevelEditorSubsystem.h
--- a/Engine/Source/Editor/LevelEditor/Public/LevelEditorSubsystem.h
+++ b/Engine/Source/Editor/LevelEditor/Public/LevelEditorSubsystem.h
@@ -171,7 +171,9 @@ public:
 	FEditorModeTools* GetLevelEditorModeManager();
 
 	//~ Begin IActorEditorContextClient interface
-	virtual void OnExecuteActorEditorContextAction(UWorld* InWorld, const EActorEditorContextAction& InType, class AActor* InActor = nullptr) {}
+	virtual void OnExecuteActorEditorContextAction(UWorld* InWorld, const EActorEditorContextAction& InType, class AActor* InActor = nullptr) override {}
+	virtual void CaptureActorEditorContextState(UWorld* InWorld, UActorEditorContextStateCollection* InStateCollection) const override {}
+	virtual void RestoreActorEditorContextState(UWorld* InWorld, const UActorEditorContextStateCollection* InStateCollection) override {}
 	virtual bool GetActorEditorContextDisplayInfo(UWorld* InWorld, FActorEditorContextClientDisplayInfo& OutDiplayInfo) const override;
 	virtual bool CanResetContext(UWorld* InWorld) const override { return false; };
 	virtual TSharedRef<SWidget> GetActorEditorContextWidget(UWorld* InWorld) const override;
diff --git a/Engine/Source/Editor/UnrealEd/Private/EditorActorFolders.cpp b/Engine/Source/Editor/UnrealEd/Private/EditorActorFolders.cpp
--- a/Engine/Source/Editor/UnrealEd/Private/EditorActorFolders.cpp
+++ b/Engine/Source/Editor/UnrealEd/Private/EditorActorFolders.cpp
@@ -606,6 +606,33 @@ void FActorFolders::OnExecuteActorEditorContextAction(UWorld* InWorld, const EAc
 	}
 }
 
+void FActorFolders::CaptureActorEditorContextState(UWorld* InWorld, UActorEditorContextStateCollection* InStateCollection) const
+{
+	const FFolder Folder = GetActorEditorContextFolder(*InWorld);
+	if (!Folder.IsNone())
+	{
+		UActorEditorContextActorFolderState* State = NewObject<UActorEditorContextActorFolderState>(InStateCollection);
+		State->FolderPath = Folder.GetPath();
+		InStateCollection->AddState(State);
+	}
+}
+
+void FActorFolders::RestoreActorEditorContextState(UWorld* InWorld, const UActorEditorContextStateCollection* InStateCollection)
+{
+	FFolder Folder = FFolder::GetWorldRootFolder(InWorld);
+
+	if (const UActorEditorContextActorFolderState* State = InStateCollection->GetState<UActorEditorContextActorFolderState>())
+	{
+		FFolder StateFolder(Folder.GetRootObject(), State->FolderPath);
+		if (ContainsFolder(*InWorld, StateFolder))
+		{
+			Folder = StateFolder;
+		}
+	}
+	
+	SetActorEditorContextFolder(*InWorld, Folder);
+}
+
 bool FActorFolders::GetActorEditorContextDisplayInfo(UWorld* InWorld, FActorEditorContextClientDisplayInfo& OutDiplayInfo) const
 {
 	const FFolder Folder = GetActorEditorContextFolder(*InWorld);
diff --git a/Engine/Source/Editor/UnrealEd/Private/Subsystems/ActorEditorContextSubsystem.cpp b/Engine/Source/Editor/UnrealEd/Private/Subsystems/ActorEditorContextSubsystem.cpp
--- a/Engine/Source/Editor/UnrealEd/Private/Subsystems/ActorEditorContextSubsystem.cpp
+++ b/Engine/Source/Editor/UnrealEd/Private/Subsystems/ActorEditorContextSubsystem.cpp
@@ -1,6 +1,7 @@
 // Copyright Epic Games, Inc. All Rights Reserved.
 
 #include "Subsystems/ActorEditorContextSubsystem.h"
+#include "ActorEditorContextState.h"
 #include "Editor/UnrealEdEngine.h"
 #include "GameFramework/Actor.h"
 #include "IActorEditorContextClient.h"
@@ -193,6 +194,36 @@ void UActorEditorContextSubsystem::InitializeContextFromActor(AActor* Actor)
 	ActorEditorContextSubsystemChanged.Broadcast();
 }
 
+void UActorEditorContextSubsystem::CaptureContext(UActorEditorContextStateCollection* InStateCollection)
+{
+	UWorld* World = GetWorld();
+	if (Clients.IsEmpty() || !World || !InStateCollection)
+	{
+		return;
+	}
+
+	InStateCollection->Reset();
+
+	for (IActorEditorContextClient* Client : Clients)
+	{
+		Client->CaptureActorEditorContextState(World, InStateCollection);
+	}
+}
+
+void UActorEditorContextSubsystem::RestoreContext(const UActorEditorContextStateCollection* InStateCollection) const
+{
+	UWorld* World = GetWorld();
+	if (Clients.IsEmpty() || !World || !InStateCollection)
+	{
+		return;
+	}
+
+	for (IActorEditorContextClient* Client : Clients)
+	{
+		Client->RestoreActorEditorContextState(World, InStateCollection);
+	}
+}
+
 TArray<IActorEditorContextClient*> UActorEditorContextSubsystem::GetDisplayableClients() const
 {
 	TArray<IActorEditorContextClient*> DisplayableClients;
diff --git a/Engine/Source/Editor/UnrealEd/Public/ActorEditorContextState.h b/Engine/Source/Editor/UnrealEd/Public/ActorEditorContextState.h
new file mode 100644
--- /dev/null
+++ b/Engine/Source/Editor/UnrealEd/Public/ActorEditorContextState.h
@@ -0,0 +1,47 @@
+// Copyright Epic Games, Inc. All Rights Reserved.
+
+#pragma once
+
+#include "CoreMinimal.h"
+#include "Templates/SubclassOf.h"
+#include "ActorEditorContextState.generated.h"
+
+UCLASS(MinimalAPI, Within = ActorEditorContextStateCollection)
+class UActorEditorContextClientState : public UObject
+{
+	GENERATED_BODY()
+};
+
+UCLASS(MinimalAPI)
+class UActorEditorContextStateCollection : public UObject
+{
+	GENERATED_BODY()
+
+public:
+	template <class TState>
+	const TState* GetState() const
+	{
+		return Cast<TState>(ClientStates.FindRef(TState::StaticClass()));
+	}
+
+	void AddState(UActorEditorContextClientState* InState)
+	{
+		ClientStates.Emplace(InState->GetClass(), InState);
+	}
+
+	bool IsEmpty() const
+	{
+		return ClientStates.IsEmpty();
+	}
+
+private:
+	void Reset()
+	{
+		ClientStates.Reset();
+	}
+
+	UPROPERTY(VisibleAnywhere, Instanced, Category = States, meta = (DisplayName = "Context"))
+	TMap<TSubclassOf<UActorEditorContextClientState>, TObjectPtr<UActorEditorContextClientState>> ClientStates;
+
+	friend class UActorEditorContextSubsystem;
+};
diff --git a/Engine/Source/Editor/UnrealEd/Public/EditorActorFolders.h b/Engine/Source/Editor/UnrealEd/Public/EditorActorFolders.h
--- a/Engine/Source/Editor/UnrealEd/Public/EditorActorFolders.h
+++ b/Engine/Source/Editor/UnrealEd/Public/EditorActorFolders.h
@@ -2,13 +2,15 @@
 
 #pragma once
 
+#include "ActorEditorContextState.h"
+#include "Folder.h"
+#include "IActorEditorContextClient.h"
 #include "UObject/ObjectMacros.h"
 #include "UObject/Object.h"
 #include "UObject/WeakObjectPtr.h"
 #include "UObject/GCObject.h"
-#include "Folder.h"
 #include "WorldFolders.h"
-#include "IActorEditorContextClient.h"
+#include "EditorActorFolders.generated.h"
 
 class FObjectPostSaveContext;
 class AActor;
@@ -29,6 +31,17 @@ DECLARE_MULTICAST_DELEGATE_TwoParams(FOnActorFolderDeleted, UWorld&, const FFold
 DECLARE_MULTICAST_DELEGATE_ThreeParams(FOnActorFolderMoved, UWorld&, const FFolder& /* src */, const FFolder& /* dst */);
 
 
+UCLASS()
+class UActorEditorContextActorFolderState : public UActorEditorContextClientState
+{
+	GENERATED_BODY()
+
+public:
+	UPROPERTY(VisibleAnywhere, Category = "Actor Folder")
+	FName FolderPath;
+};
+
+
 /** Class responsible for managing an in-memory representation of actor folders in the editor */
 struct FActorFolders : public FGCObject, public IActorEditorContextClient
 {
@@ -162,10 +175,12 @@ struct FActorFolders : public FGCObject, public IActorEditorContextClient
 	UNREALED_API FActorFolderProps* GetFolderProperties(UWorld& InWorld, const FFolder& InFolder);
 
 	//~ Begin IActorEditorContextClient interface
-	UNREALED_API virtual void OnExecuteActorEditorContextAction(UWorld* InWorld, const EActorEditorContextAction& InType, class AActor* InActor = nullptr) override;
-	UNREALED_API virtual bool GetActorEditorContextDisplayInfo(UWorld* InWorld, FActorEditorContextClientDisplayInfo& OutDiplayInfo) const override;
+	virtual void OnExecuteActorEditorContextAction(UWorld* InWorld, const EActorEditorContextAction& InType, class AActor* InActor = nullptr) override;
+	virtual void CaptureActorEditorContextState(UWorld* InWorld, UActorEditorContextStateCollection* InStateCollection) const override;
+	virtual void RestoreActorEditorContextState(UWorld* InWorld, const UActorEditorContextStateCollection* InStateCollection) override;
+	virtual bool GetActorEditorContextDisplayInfo(UWorld* InWorld, FActorEditorContextClientDisplayInfo& OutDiplayInfo) const override;
 	virtual bool CanResetContext(UWorld* InWorld) const override { return true; };
-	UNREALED_API virtual TSharedRef<SWidget> GetActorEditorContextWidget(UWorld* InWorld) const override;
+	virtual TSharedRef<SWidget> GetActorEditorContextWidget(UWorld* InWorld) const override;
 	virtual FOnActorEditorContextClientChanged& GetOnActorEditorContextClientChanged() override { return ActorEditorContextClientChanged; }
 	//~ End IActorEditorContextClient interface
 	UNREALED_API FFolder GetActorEditorContextFolder(UWorld& InWorld, bool bMustMatchCurrentLevel = true) const;
diff --git a/Engine/Source/Editor/UnrealEd/Public/IActorEditorContextClient.h b/Engine/Source/Editor/UnrealEd/Public/IActorEditorContextClient.h
--- a/Engine/Source/Editor/UnrealEd/Public/IActorEditorContextClient.h
+++ b/Engine/Source/Editor/UnrealEd/Public/IActorEditorContextClient.h
@@ -13,6 +13,7 @@ DECLARE_MULTICAST_DELEGATE_OneParam(FOnActorEditorContextClientChanged, struct I
 
 class UWorld;
 class AActor;
+class UActorEditorContextStateCollection;
 
 struct FActorEditorContextClientDisplayInfo
 {
@@ -34,6 +35,8 @@ enum class EActorEditorContextAction
 struct IActorEditorContextClient
 {
 	virtual void OnExecuteActorEditorContextAction(UWorld* InWorld, const EActorEditorContextAction& InType, AActor* InActor = nullptr) = 0;
+	virtual void CaptureActorEditorContextState(UWorld* InWorld, UActorEditorContextStateCollection* InStateCollection) const = 0;
+	virtual void RestoreActorEditorContextState(UWorld* InWorld, const UActorEditorContextStateCollection* InStateCollection) = 0;
 	virtual bool GetActorEditorContextDisplayInfo(UWorld* InWorld, FActorEditorContextClientDisplayInfo& OutDiplayInfo) const = 0;
 	virtual bool CanResetContext(UWorld* InWorld) const = 0;
 	virtual TSharedRef<SWidget> GetActorEditorContextWidget(UWorld* InWorld) const = 0;
diff --git a/Engine/Source/Editor/UnrealEd/Public/Subsystems/ActorEditorContextSubsystem.h b/Engine/Source/Editor/UnrealEd/Public/Subsystems/ActorEditorContextSubsystem.h
--- a/Engine/Source/Editor/UnrealEd/Public/Subsystems/ActorEditorContextSubsystem.h
+++ b/Engine/Source/Editor/UnrealEd/Public/Subsystems/ActorEditorContextSubsystem.h
@@ -9,6 +9,7 @@
 
 struct IActorEditorContextClient;
 class AActor;
+class UActorEditorContextStateCollection;
 
 DECLARE_MULTICAST_DELEGATE(FOnActorEditorContextSubsystemChanged);
 
@@ -34,6 +35,8 @@ public:
 	UNREALED_API void PushContext(bool bDuplicateContext = false);
 	UNREALED_API void PopContext();
 	UNREALED_API void InitializeContextFromActor(AActor* Actor);
+	UNREALED_API void CaptureContext(UActorEditorContextStateCollection* InStateCollection);
+	UNREALED_API void RestoreContext(const UActorEditorContextStateCollection* InStateCollection) const;
 	UNREALED_API TArray<IActorEditorContextClient*> GetDisplayableClients() const;
 	FOnActorEditorContextSubsystemChanged& OnActorEditorContextSubsystemChanged() { return ActorEditorContextSubsystemChanged; }
 
diff --git a/Engine/Source/Editor/WorldPartitionEditor/Private/WorldPartition/ContentBundle/ContentBundleEditorSubsystem.cpp b/Engine/Source/Editor/WorldPartitionEditor/Private/WorldPartition/ContentBundle/ContentBundleEditorSubsystem.cpp
--- a/Engine/Source/Editor/WorldPartitionEditor/Private/WorldPartition/ContentBundle/ContentBundleEditorSubsystem.cpp
+++ b/Engine/Source/Editor/WorldPartitionEditor/Private/WorldPartition/ContentBundle/ContentBundleEditorSubsystem.cpp
@@ -19,6 +19,7 @@
 #include "Widgets/SBoxPanel.h"
 #include "Widgets/Text/STextBlock.h"
 #include "Editor.h"
+#include "ActorEditorContextState.h"
 
 #define LOCTEXT_NAMESPACE "ContentBundle"
 
@@ -152,6 +153,35 @@ void UContentBundleEditingSubmodule::OnExecuteActorEditorContextAction(UWorld* I
 	}
 }
 
+void UContentBundleEditingSubmodule::CaptureActorEditorContextState(UWorld* InWorld, UActorEditorContextStateCollection* InStateCollection) const
+{
+	if (EditingContentBundleGuid.IsValid())
+	{
+		UActorEditorContextContentBundleState* State = NewObject<UActorEditorContextContentBundleState>(InStateCollection);
+		State->ContentBundleGuid = EditingContentBundleGuid;
+		InStateCollection->AddState(State);
+	}
+}
+
+void UContentBundleEditingSubmodule::RestoreActorEditorContextState(UWorld* InWorld, const UActorEditorContextStateCollection* InStateCollection)
+{
+	TSharedPtr<FContentBundleEditor> ContentBundleToActivate;
+
+	if (const UActorEditorContextContentBundleState* State = InStateCollection->GetState<UActorEditorContextContentBundleState>())
+	{
+		ContentBundleToActivate = GetEditorContentBundle(State->ContentBundleGuid);
+	}
+
+	if (ContentBundleToActivate.IsValid())
+	{
+		ActivateContentBundleEditing(ContentBundleToActivate);
+	}
+	else
+	{
+		DeactivateCurrentContentBundleEditing();
+	}
+}
+
 bool UContentBundleEditingSubmodule::GetActorEditorContextDisplayInfo(UWorld* InWorld, FActorEditorContextClientDisplayInfo& OutDiplayInfo) const
 {
 	if (EditingContentBundleGuid.IsValid())
diff --git a/Engine/Source/Editor/WorldPartitionEditor/Public/WorldPartition/ContentBundle/ContentBundleEditorSubsystem.h b/Engine/Source/Editor/WorldPartitionEditor/Public/WorldPartition/ContentBundle/ContentBundleEditorSubsystem.h
--- a/Engine/Source/Editor/WorldPartitionEditor/Public/WorldPartition/ContentBundle/ContentBundleEditorSubsystem.h
+++ b/Engine/Source/Editor/WorldPartitionEditor/Public/WorldPartition/ContentBundle/ContentBundleEditorSubsystem.h
@@ -1,13 +1,13 @@
 // Copyright Epic Games, Inc. All Rights Reserved.
 #pragma once
 
+#include "ActorEditorContextState.h"
 #include "CoreMinimal.h"
+#include "EditorSubsystem.h"
+#include "IActorEditorContextClient.h"
 #include "UObject/ObjectMacros.h"
 #include "UObject/Object.h"
-#include "EditorSubsystem.h"
 #include "WorldPartition/ContentBundle/ContentBundleEditorSubsystemInterface.h"
-#include "IActorEditorContextClient.h"
-
 #include "ContentBundleEditorSubsystem.generated.h"
 
 class FContentBundleEditor;
@@ -39,6 +39,16 @@ private:
 	bool bIsInitialized = false;
 };
 
+UCLASS()
+class UActorEditorContextContentBundleState : public UActorEditorContextClientState
+{
+	GENERATED_BODY()
+
+public:
+	UPROPERTY(VisibleAnywhere, Category = "Content Bundle")
+	FGuid ContentBundleGuid;
+};
+
 UCLASS()
 class UContentBundleEditingSubmodule : public UContentBundleEditorSubsystemModule, public IActorEditorContextClient
 {
@@ -49,6 +59,8 @@ class UContentBundleEditingSubmodule : public UContentBundleEditorSubsystemModul
 public:
 	//~ Begin IActorEditorContextClient interface
 	virtual void OnExecuteActorEditorContextAction(UWorld* InWorld, const EActorEditorContextAction& InType, AActor* InActor = nullptr) override;
+	virtual void CaptureActorEditorContextState(UWorld* InWorld, UActorEditorContextStateCollection* InStateCollection) const override;
+	virtual void RestoreActorEditorContextState(UWorld* InWorld, const UActorEditorContextStateCollection* InStateCollection) override;
 	virtual bool GetActorEditorContextDisplayInfo(UWorld* InWorld, FActorEditorContextClientDisplayInfo& OutDiplayInfo) const override;
 	virtual bool CanResetContext(UWorld* InWorld) const override { return true; }
 	virtual TSharedRef<SWidget> GetActorEditorContextWidget(UWorld* InWorld) const override;
