Added BP Functions SetDataLayerInitialRuntimeState & SetDataLayerIsInitiallyVisible to UDataLayerEditorSubsystem.

diff --git a/Engine/Source/Editor/DataLayerEditor/Private/DataLayer/DataLayerEditorSubsystem.cpp b/Engine/Source/Editor/DataLayerEditor/Private/DataLayer/DataLayerEditorSubsystem.cpp
--- a/Engine/Source/Editor/DataLayerEditor/Private/DataLayer/DataLayerEditorSubsystem.cpp
+++ b/Engine/Source/Editor/DataLayerEditor/Private/DataLayer/DataLayerEditorSubsystem.cpp
@@ -984,6 +984,22 @@ bool UDataLayerEditorSubsystem::SetParentDataLayer(UDataLayerInstance* DataLayer
 	return SetParentDataLayerForDataLayersInternal({ DataLayerInstance }, ParentDataLayer);
 }
 
+void UDataLayerEditorSubsystem::SetDataLayerInitialRuntimeState(UDataLayerInstance* DataLayerInstance, EDataLayerRuntimeState InitialRuntimeState)
+{
+	if (DataLayerInstance)
+	{
+		DataLayerInstance->SetInitialRuntimeState(InitialRuntimeState);
+	}
+}
+
+void UDataLayerEditorSubsystem::SetDataLayerIsInitiallyVisible(UDataLayerInstance* DataLayerInstance, bool bIsInitiallyVisible)
+{
+	if (DataLayerInstance)
+	{
+		DataLayerInstance->SetIsInitiallyVisible(bIsInitiallyVisible);
+	}
+}
+
 bool UDataLayerEditorSubsystem::AddActorToDataLayer(AActor* Actor, UDataLayerInstance* DataLayerInstance)
 {
 	TArray<AActor*> Actors;
diff --git a/Engine/Source/Editor/DataLayerEditor/Public/DataLayer/DataLayerEditorSubsystem.h b/Engine/Source/Editor/DataLayerEditor/Public/DataLayer/DataLayerEditorSubsystem.h
--- a/Engine/Source/Editor/DataLayerEditor/Public/DataLayer/DataLayerEditorSubsystem.h
+++ b/Engine/Source/Editor/DataLayerEditor/Public/DataLayer/DataLayerEditorSubsystem.h
@@ -636,6 +636,24 @@ public:
 	UFUNCTION(BlueprintCallable, Category = DataLayers)
 	void SetParentDataLayerForDataLayers(const TArray<UDataLayerInstance*>& DataLayers, UDataLayerInstance* ParentDataLayer);
 
+	/**
+	 * Sets the initial runtime state of the DataLayers
+	 *
+	 * @param	DataLayer						The DataLayer to affect
+	 * @param	InitialRuntimeState				The initial runtime state.
+	 */
+	UFUNCTION(BlueprintCallable, Category = DataLayers)
+	void SetDataLayerInitialRuntimeState(UDataLayerInstance* DataLayer, EDataLayerRuntimeState InitialRuntimeState);
+
+	/**
+	 * Sets the initial editor visibility state of the DataLayers
+	 *
+	 * @param	DataLayer						The DataLayer to affect
+	 * @param	InitialRuntimeState				The initial editor visibility state.
+	 */
+	UFUNCTION(BlueprintCallable, Category = DataLayers)
+	void SetDataLayerIsInitiallyVisible(UDataLayerInstance* DataLayer, bool bIsInitiallyVisible);
+
 	/**
 	 * Deletes all of the provided DataLayers
 	 *
diff --git a/Engine/Source/Runtime/Engine/Private/WorldPartition/DataLayer/DataLayerInstance.cpp b/Engine/Source/Runtime/Engine/Private/WorldPartition/DataLayer/DataLayerInstance.cpp
--- a/Engine/Source/Runtime/Engine/Private/WorldPartition/DataLayer/DataLayerInstance.cpp
+++ b/Engine/Source/Runtime/Engine/Private/WorldPartition/DataLayer/DataLayerInstance.cpp
@@ -494,6 +494,14 @@ bool UDataLayerInstance::IsParentDataLayerTypeCompatible(const UDataLayerInstanc
 	return true;
 }
 
+void UDataLayerInstance::SetInitialRuntimeState(EDataLayerRuntimeState InInitialRuntimeState)
+{
+	if (IsRuntime())
+	{
+		InitialRuntimeState = InInitialRuntimeState;
+	}
+}
+
 bool UDataLayerInstance::SetParent(UDataLayerInstance* InParent)
 {
 	if (!CanBeChildOf(InParent))
diff --git a/Engine/Source/Runtime/Engine/Public/WorldPartition/DataLayer/DataLayerInstance.h b/Engine/Source/Runtime/Engine/Public/WorldPartition/DataLayer/DataLayerInstance.h
--- a/Engine/Source/Runtime/Engine/Public/WorldPartition/DataLayer/DataLayerInstance.h
+++ b/Engine/Source/Runtime/Engine/Public/WorldPartition/DataLayer/DataLayerInstance.h
@@ -95,7 +95,8 @@ public:
 	ENGINE_API void SetVisible(bool bIsVisible);
 	ENGINE_API void SetIsInitiallyVisible(bool bIsInitiallyVisible);
 	ENGINE_API void SetIsLoadedInEditor(bool bIsLoadedInEditor, bool bFromUserChange);
-	void SetIsLocked(bool bInIsLocked) { bIsLocked = bInIsLocked; }
+	ENGINE_API void SetIsLocked(bool bInIsLocked) { bIsLocked = bInIsLocked; }
+	ENGINE_API void SetInitialRuntimeState(EDataLayerRuntimeState InInitialRuntimeState);
 
 	ENGINE_API bool IsEffectiveLoadedInEditor() const;
 	ENGINE_API bool IsInitiallyLoadedInEditor() const { return bIsInitiallyLoadedInEditor; }
