[WorldPartition][DataLayer] Data Layers Streaming Priority
* Add Streaming Priority property to Data Layer Instance
* Determines the streaming priority (lower value means higher priority), relative to target grid cells. Priority 0 means the same priority as target grid cells.
* Maintains existing sorting criteria: source priority -> cell hierarchical level -> cell priority (affected by this change) -> cell distance and angle from source
* If a cell consists of more than one Data Layer, the highest priority is applied
* Add text column to Data Layer Browser displaying streaming priority

diff --git a/Engine/Source/Editor/DataLayerEditor/Private/DataLayer/SDataLayerBrowser.cpp b/Engine/Source/Editor/DataLayerEditor/Private/DataLayer/SDataLayerBrowser.cpp
--- a/Engine/Source/Editor/DataLayerEditor/Private/DataLayer/SDataLayerBrowser.cpp
+++ b/Engine/Source/Editor/DataLayerEditor/Private/DataLayer/SDataLayerBrowser.cpp
@@ -123,6 +123,21 @@ void SDataLayerBrowser::Construct(const FArguments& InArgs)
 		return FString();
 	});
 
+	FGetTextForItem InternalStreamingPriorityInfoText = FGetTextForItem::CreateLambda([](const ISceneOutlinerTreeItem& Item) -> FString
+	{
+			if (const FDataLayerTreeItem* DataLayerItem = Item.CastTo<FDataLayerTreeItem>())
+			{
+				if (const UDataLayerInstance* DataLayerInstance = DataLayerItem->GetDataLayer())
+				{
+					if (DataLayerInstance->GetStreamingPriority() != 0)
+					{
+						return FString::FromInt(DataLayerInstance->GetStreamingPriority());
+					}
+				}
+			}
+			return FString();
+	});
+
 	SAssignNew(DeprecatedDataLayerWarningBox, SMultiLineEditableTextBox)
 		.IsReadOnly(true)
 		.Font(IDetailLayoutBuilder::GetDetailFontBold())
@@ -146,6 +161,7 @@ void SDataLayerBrowser::Construct(const FArguments& InArgs)
 	InitOptions.ColumnMap.Add("Initial State", FSceneOutlinerColumnInfo(ESceneOutlinerColumnVisibility::Invisible, 20, FCreateSceneOutlinerColumn::CreateStatic(&FTextInfoColumn::CreateTextInfoColumn, FName("Initial State"), InternalInitialRuntimeStateInfoText, FText::FromString("Initial Runtime State"))));
 	InitOptions.ColumnMap.Add(FDataLayerOutlinerHasErrorsColumn::GetID(), FSceneOutlinerColumnInfo(ESceneOutlinerColumnVisibility::Visible, 100, FCreateSceneOutlinerColumn::CreateLambda([](ISceneOutliner& InSceneOutliner) { return MakeShareable(new FDataLayerOutlinerHasErrorsColumn(InSceneOutliner)); }), false));
 	InitOptions.ColumnMap.Add(FDataLayerOutlinerDebugColorColumn::GetID(), FSceneOutlinerColumnInfo(ESceneOutlinerColumnVisibility::Invisible, 100, FCreateSceneOutlinerColumn::CreateLambda([](ISceneOutliner& InSceneOutliner) { return MakeShareable(new FDataLayerOutlinerDebugColorColumn(InSceneOutliner)); })));
+	InitOptions.ColumnMap.Add("Streaming Priority", FSceneOutlinerColumnInfo(ESceneOutlinerColumnVisibility::Invisible, 20, FCreateSceneOutlinerColumn::CreateStatic(&FTextInfoColumn::CreateTextInfoColumn, FName("Streaming Priority"), InternalStreamingPriorityInfoText, FText::FromString("Streaming Priority"))));
 	DataLayerOutliner = SNew(SDataLayerOutliner, InitOptions).IsEnabled(FSlateApplication::Get().GetNormalExecutionAttribute());
 
 	SAssignNew(DataLayerContentsSection, SBorder)
diff --git a/Engine/Source/Runtime/Engine/Private/WorldPartition/DataLayer/DataLayerInstance.cpp b/Engine/Source/Runtime/Engine/Private/WorldPartition/DataLayer/DataLayerInstance.cpp
--- a/Engine/Source/Runtime/Engine/Private/WorldPartition/DataLayer/DataLayerInstance.cpp
+++ b/Engine/Source/Runtime/Engine/Private/WorldPartition/DataLayer/DataLayerInstance.cpp
@@ -29,6 +29,7 @@ UDataLayerInstance::UDataLayerInstance(const FObjectInitializer& ObjectInitializ
 	, bIsLoadedInEditor(true)
 	, bIsLocked(false)
 	, OverrideBlockOnSlowStreaming(EOverrideBlockOnSlowStreaming::NoOverride)
+	, StreamingPriority(0)
 #endif
 	, InitialRuntimeState(EDataLayerRuntimeState::Unloaded)
 {}
diff --git a/Engine/Source/Runtime/Engine/Private/WorldPartition/RuntimeHashSet/WorldPartitionRuntimeHashSetStreamingGeneration.cpp b/Engine/Source/Runtime/Engine/Private/WorldPartition/RuntimeHashSet/WorldPartitionRuntimeHashSetStreamingGeneration.cpp
--- a/Engine/Source/Runtime/Engine/Private/WorldPartition/RuntimeHashSet/WorldPartitionRuntimeHashSetStreamingGeneration.cpp
+++ b/Engine/Source/Runtime/Engine/Private/WorldPartition/RuntimeHashSet/WorldPartitionRuntimeHashSetStreamingGeneration.cpp
@@ -158,7 +158,8 @@ bool UWorldPartitionRuntimeHashSet::GenerateStreaming(UWorldPartitionStreamingPo
 		RuntimeCellData->DebugName = CellUniqueId.Name;
 		RuntimeCellData->CellBounds = CellDescInstance.CellBounds;
 		RuntimeCellData->HierarchicalLevel = CellDescInstance.bIsSpatiallyLoaded ? CellDescInstance.Level : MAX_int32;
-		RuntimeCellData->Priority = CellDescInstance.Priority;
+		const int32 DataLayersStreamingPriority = GetDataLayersStreamingPriority(CellDescInstance.DataLayerInstances);
+		RuntimeCellData->Priority = CellDescInstance.Priority + DataLayersStreamingPriority;
 		RuntimeCellData->GridName = CellDescInstance.SourcePartition->Name;
 		RuntimeCellData->bIs2D = CellDescInstance.bIs2D;
 
diff --git a/Engine/Source/Runtime/Engine/Private/WorldPartition/WorldPartitionRuntimeHash.cpp b/Engine/Source/Runtime/Engine/Private/WorldPartition/WorldPartitionRuntimeHash.cpp
--- a/Engine/Source/Runtime/Engine/Private/WorldPartition/WorldPartitionRuntimeHash.cpp
+++ b/Engine/Source/Runtime/Engine/Private/WorldPartition/WorldPartitionRuntimeHash.cpp
@@ -555,6 +555,44 @@ bool UWorldPartitionRuntimeHash::ResolveBlockOnSlowStreamingForCell(bool bInOwne
 	const bool bBlockOnSlowStreaming = DataLayersOverrideBlockOnSlowStreaming.IsSet() ? DataLayersOverrideBlockOnSlowStreaming.GetValue() : bInOwnerBlockOnSlowStreaming;
 	return bBlockOnSlowStreaming;
 }
+
+int32 UWorldPartitionRuntimeHash::GetDataLayersStreamingPriority(const TArray<const UDataLayerInstance*>& InCellDataLayerInstances) const
+{
+	if (InCellDataLayerInstances.Num() == 0)
+	{
+		return 0;
+	}
+
+	bool bDetectedMismatch = false;
+	int32 StreamingPriority = InCellDataLayerInstances[0]->GetStreamingPriority();
+	for (const UDataLayerInstance* DataLayerInstance : InCellDataLayerInstances)
+	{
+		bDetectedMismatch |= (StreamingPriority != DataLayerInstance->GetStreamingPriority());
+		StreamingPriority = FMath::Min(StreamingPriority, DataLayerInstance->GetStreamingPriority());
+	}
+
+#if !NO_LOGGING
+	if (bDetectedMismatch)
+	{
+		FStringBuilderBase HighestPriorityLayers;
+		FStringBuilderBase MismatchedLayers;
+		for (const UDataLayerInstance* DataLayerInstance : InCellDataLayerInstances)
+		{
+			if (StreamingPriority == DataLayerInstance->GetStreamingPriority())
+			{
+				HighestPriorityLayers.Appendf(TEXT("%s "), *DataLayerInstance->GetDataLayerShortName());
+			}
+			else
+			{
+				MismatchedLayers.Appendf(TEXT("%s "), *DataLayerInstance->GetDataLayerShortName());
+			}
+		}
+		UE_LOG(LogWorldPartition, Log, TEXT("Found data layers streaming priority mismatch. The following data layers (%s) have lower priority than (%s) - using the highest found priority (%d) for this cell."), MismatchedLayers.ToString(), HighestPriorityLayers.ToString(), StreamingPriority);
+	}
+#endif
+
+	return StreamingPriority;
+}
 #endif
 
 bool UWorldPartitionRuntimeHash::IsCellRelevantFor(bool bClientOnlyVisible) const
diff --git a/Engine/Source/Runtime/Engine/Private/WorldPartition/WorldPartitionRuntimeSpatialHash.cpp b/Engine/Source/Runtime/Engine/Private/WorldPartition/WorldPartitionRuntimeSpatialHash.cpp
--- a/Engine/Source/Runtime/Engine/Private/WorldPartition/WorldPartitionRuntimeSpatialHash.cpp
+++ b/Engine/Source/Runtime/Engine/Private/WorldPartition/WorldPartitionRuntimeSpatialHash.cpp
@@ -1652,7 +1652,8 @@ bool UWorldPartitionRuntimeSpatialHash::CreateStreamingGrid(const FSpatialHashRu
 					CellDataSpatialHash->Extent = (float)CellExtent;
 					CellDataSpatialHash->GridName = RuntimeGrid.GridName;
 					CellDataSpatialHash->DebugName = CellName + WorldInstanceSuffix;
-					CellDataSpatialHash->Priority = RuntimeGrid.Priority;
+					const int32 DataLayersStreamingPriority = GetDataLayersStreamingPriority(DataLayers);
+					CellDataSpatialHash->Priority = RuntimeGrid.Priority + DataLayersStreamingPriority;
 
 					PopulateRuntimeCell(StreamingCell, FilteredActors, OutPackagesToGenerate);
 
diff --git a/Engine/Source/Runtime/Engine/Public/WorldPartition/DataLayer/DataLayerInstance.h b/Engine/Source/Runtime/Engine/Public/WorldPartition/DataLayer/DataLayerInstance.h
--- a/Engine/Source/Runtime/Engine/Public/WorldPartition/DataLayer/DataLayerInstance.h
+++ b/Engine/Source/Runtime/Engine/Public/WorldPartition/DataLayer/DataLayerInstance.h
@@ -105,6 +105,7 @@ public:
 	ENGINE_API virtual bool IsReadOnly(FText* OutReason = nullptr) const;
 	virtual bool IsIncludedInActorFilterDefault() const { return false; }
 	ENGINE_API EOverrideBlockOnSlowStreaming GetOverrideBlockOnSlowStreaming() const { return OverrideBlockOnSlowStreaming; }
+	ENGINE_API int32 GetStreamingPriority() const { return StreamingPriority; }
 
 	// Data Layer Instance features support
 	ENGINE_API bool CanBeChildOf(const UDataLayerInstance* InParent, FText* OutReason = nullptr) const;
@@ -258,6 +259,10 @@ protected:
 
 	UPROPERTY(Category = "Runtime|Advanced", EditAnywhere)
 	EOverrideBlockOnSlowStreaming OverrideBlockOnSlowStreaming;
+
+	/** Determines the streaming priority (lower value means higher priority), relative to target grid cells. Priority 0 means the same priority as target grid cells. */
+	UPROPERTY(Category = "Runtime|Advanced", EditAnywhere)
+	int32 StreamingPriority;
 #endif
 
 	/** Initial runtime state of this data layer instance. Only supported if it's runtime and not client/server only. */
diff --git a/Engine/Source/Runtime/Engine/Public/WorldPartition/WorldPartitionRuntimeHash.h b/Engine/Source/Runtime/Engine/Public/WorldPartition/WorldPartitionRuntimeHash.h
--- a/Engine/Source/Runtime/Engine/Public/WorldPartition/WorldPartitionRuntimeHash.h
+++ b/Engine/Source/Runtime/Engine/Public/WorldPartition/WorldPartitionRuntimeHash.h
@@ -308,6 +308,7 @@ protected:
 #if WITH_EDITOR
 	ENGINE_API void ForceExternalActorLevelReference(bool bForceExternalActorLevelReference);
 	ENGINE_API bool ResolveBlockOnSlowStreamingForCell(bool bInOwnerBlockOnSlowStreaming, bool bInIsHLODCell, const TArray<const UDataLayerInstance*>& InCellDataLayerInstances) const;
+	ENGINE_API int32 GetDataLayersStreamingPriority(const TArray<const UDataLayerInstance*>& InCellDataLayerInstances) const;
 #endif
 
 	TSet<TWeakObjectPtr<URuntimeHashExternalStreamingObjectBase>> InjectedExternalStreamingObjects;
