Data Layers
* Outliner - Make sure headers and tooltips are localized.

#rnx

diff --git a/Engine/Source/Editor/DataLayerEditor/Private/DataLayer/DataLayerOutlinerIsLoadedInEditorColumn.cpp b/Engine/Source/Editor/DataLayerEditor/Private/DataLayer/DataLayerOutlinerIsLoadedInEditorColumn.cpp
--- a/Engine/Source/Editor/DataLayerEditor/Private/DataLayer/DataLayerOutlinerIsLoadedInEditorColumn.cpp
+++ b/Engine/Source/Editor/DataLayerEditor/Private/DataLayer/DataLayerOutlinerIsLoadedInEditorColumn.cpp
@@ -45,7 +45,7 @@ SHeaderRow::FColumn::FArguments FDataLayerOutlinerIsLoadedInEditorColumn::Constr
 		.VAlignHeader(VAlign_Center)
 		.HAlignCell(HAlign_Center)
 		.VAlignCell(VAlign_Center)
-		.DefaultTooltip(FText::FromName(GetColumnID()))
+		.DefaultTooltip(LOCTEXT("DataLayerLoadedInEditorTooltip", "Data Layer Loaded In Editor"))
 		[
 			SNew(SImage)
 			.Image(FAppStyle::GetBrush(TEXT("DataLayer.LoadedInEditor")))
diff --git a/Engine/Source/Editor/DataLayerEditor/Private/DataLayer/SDataLayerBrowser.cpp b/Engine/Source/Editor/DataLayerEditor/Private/DataLayer/SDataLayerBrowser.cpp
--- a/Engine/Source/Editor/DataLayerEditor/Private/DataLayer/SDataLayerBrowser.cpp
+++ b/Engine/Source/Editor/DataLayerEditor/Private/DataLayer/SDataLayerBrowser.cpp
@@ -155,13 +155,13 @@ void SDataLayerBrowser::Construct(const FArguments& InArgs)
 	InitOptions.ModeFactory = FCreateSceneOutlinerMode::CreateLambda([this](SSceneOutliner* Outliner) { return new FDataLayerMode(FDataLayerModeParams(Outliner, this, nullptr)); });
 	InitOptions.ColumnMap.Add(FDataLayerOutlinerIsVisibleColumn::GetID(), FSceneOutlinerColumnInfo(ESceneOutlinerColumnVisibility::Visible, 0, FCreateSceneOutlinerColumn::CreateLambda([](ISceneOutliner& InSceneOutliner) { return MakeShareable(new FDataLayerOutlinerIsVisibleColumn(InSceneOutliner)); })));
 	InitOptions.ColumnMap.Add(FDataLayerOutlinerIsLoadedInEditorColumn::GetID(), FSceneOutlinerColumnInfo(ESceneOutlinerColumnVisibility::Visible, 1, FCreateSceneOutlinerColumn::CreateLambda([](ISceneOutliner& InSceneOutliner) { return MakeShareable(new FDataLayerOutlinerIsLoadedInEditorColumn(InSceneOutliner)); })));
-	InitOptions.ColumnMap.Add(FSceneOutlinerBuiltInColumnTypes::Label(), FSceneOutlinerColumnInfo(ESceneOutlinerColumnVisibility::Visible, 2));
+	InitOptions.ColumnMap.Add(FSceneOutlinerBuiltInColumnTypes::Label(), FSceneOutlinerColumnInfo(ESceneOutlinerColumnVisibility::Visible, 2, FCreateSceneOutlinerColumn(), false, TOptional<float>(), FSceneOutlinerBuiltInColumnTypes::Label_Localized()));
 	InitOptions.ColumnMap.Add(FDataLayerOutlinerDeleteButtonColumn::GetID(), FSceneOutlinerColumnInfo(ESceneOutlinerColumnVisibility::Visible, 10, FCreateSceneOutlinerColumn::CreateLambda([](ISceneOutliner& InSceneOutliner) { return MakeShareable(new FDataLayerOutlinerDeleteButtonColumn(InSceneOutliner)); })));
-	InitOptions.ColumnMap.Add("ID Name", FSceneOutlinerColumnInfo(ESceneOutlinerColumnVisibility::Invisible, 20, FCreateSceneOutlinerColumn::CreateStatic(&FTextInfoColumn::CreateTextInfoColumn, FName("ID Name"), InternalNameInfoText, FText::GetEmpty())));
-	InitOptions.ColumnMap.Add("Initial State", FSceneOutlinerColumnInfo(ESceneOutlinerColumnVisibility::Invisible, 20, FCreateSceneOutlinerColumn::CreateStatic(&FTextInfoColumn::CreateTextInfoColumn, FName("Initial State"), InternalInitialRuntimeStateInfoText, FText::FromString("Initial Runtime State"))));
+	InitOptions.ColumnMap.Add("ID Name", FSceneOutlinerColumnInfo(ESceneOutlinerColumnVisibility::Invisible, 20, FCreateSceneOutlinerColumn::CreateStatic(&FTextInfoColumn::CreateTextInfoColumn, FName("ID Name"), InternalNameInfoText, LOCTEXT("DataLayerIDNameToolTip", "ID Name")), true, TOptional<float>(), LOCTEXT("DataLayerIDName", "ID Name")));
+	InitOptions.ColumnMap.Add("Initial State", FSceneOutlinerColumnInfo(ESceneOutlinerColumnVisibility::Invisible, 20, FCreateSceneOutlinerColumn::CreateStatic(&FTextInfoColumn::CreateTextInfoColumn, FName("Initial State"), InternalInitialRuntimeStateInfoText, LOCTEXT("DataLayerInitialRuntimeStateToolTip", "Initial Runtime State")), true, TOptional<float>(), LOCTEXT("DataLayerInitialState", "Initial State")));
 	InitOptions.ColumnMap.Add(FDataLayerOutlinerHasErrorsColumn::GetID(), FSceneOutlinerColumnInfo(ESceneOutlinerColumnVisibility::Visible, 100, FCreateSceneOutlinerColumn::CreateLambda([](ISceneOutliner& InSceneOutliner) { return MakeShareable(new FDataLayerOutlinerHasErrorsColumn(InSceneOutliner)); }), false));
 	InitOptions.ColumnMap.Add(FDataLayerOutlinerDebugColorColumn::GetID(), FSceneOutlinerColumnInfo(ESceneOutlinerColumnVisibility::Invisible, 100, FCreateSceneOutlinerColumn::CreateLambda([](ISceneOutliner& InSceneOutliner) { return MakeShareable(new FDataLayerOutlinerDebugColorColumn(InSceneOutliner)); })));
-	InitOptions.ColumnMap.Add("Streaming Priority", FSceneOutlinerColumnInfo(ESceneOutlinerColumnVisibility::Invisible, 20, FCreateSceneOutlinerColumn::CreateStatic(&FTextInfoColumn::CreateTextInfoColumn, FName("Streaming Priority"), InternalStreamingPriorityInfoText, FText::FromString("Streaming Priority"))));
+	InitOptions.ColumnMap.Add("Streaming Priority", FSceneOutlinerColumnInfo(ESceneOutlinerColumnVisibility::Invisible, 20, FCreateSceneOutlinerColumn::CreateStatic(&FTextInfoColumn::CreateTextInfoColumn, FName("Streaming Priority"), InternalStreamingPriorityInfoText, LOCTEXT("DataLayerStreamingPriorityToolTip", "Streaming Priority")), true, TOptional<float>(), LOCTEXT("DataLayerStreamingPriority", "Streaming Priority")));
 	DataLayerOutliner = SNew(SDataLayerOutliner, InitOptions).IsEnabled(FSlateApplication::Get().GetNormalExecutionAttribute());
 
 	SAssignNew(DataLayerContentsSection, SBorder)
