Updated headers using UnrealCodeFixup to make sure dllstorage is on methods/staticvars instead of types. Part 5/n

diff --git a/Engine/Source/Editor/DataHierarchyEditor/Public/DataHierarchyViewModelBase.h b/Engine/Source/Editor/DataHierarchyEditor/Public/DataHierarchyViewModelBase.h
--- a/Engine/Source/Editor/DataHierarchyEditor/Public/DataHierarchyViewModelBase.h
+++ b/Engine/Source/Editor/DataHierarchyEditor/Public/DataHierarchyViewModelBase.h
@@ -16,6 +16,8 @@
 #include "UObject/UObjectGlobals.h"
 #include "DataHierarchyViewModelBase.generated.h"
 
+#define UE_API DATAHIERARCHYEDITOR_API
+
 /** HIERARCHY EDITOR
  *	The hierarchy editor is a generic tool to organize and structure all kinds of data.
  *	It inherently supports sections, categories, and items. You can add your own items and customize how they are visualized.
@@ -135,14 +137,14 @@ FORCEINLINE uint32 GetTypeHash(const FHierarchyElementIdentity& Identity)
 }
 
 /** A base class that is used to refresh data that represents external data. Inherit from this class if you need more context data. */
-UCLASS(Transient)
-class DATAHIERARCHYEDITOR_API UHierarchyDataRefreshContext : public UObject
+UCLASS(MinimalAPI, Transient)
+class UHierarchyDataRefreshContext : public UObject
 {
 	GENERATED_BODY()
 };
 
-UCLASS()
-class DATAHIERARCHYEDITOR_API UHierarchyElement : public UObject
+UCLASS(MinimalAPI)
+class UHierarchyElement : public UObject
 {
 	GENERATED_BODY()
 
@@ -151,18 +153,18 @@ public:
 	virtual ~UHierarchyElement() override {}
 
 	TArray<TObjectPtr<UHierarchyElement>>& GetChildrenMutable() { return Children; }
-	const TArray<UHierarchyElement*>& GetChildren() const;
+	UE_API const TArray<UHierarchyElement*>& GetChildren() const;
 
 	template<class ChildClass>
 	ChildClass* AddChild();
 
-	UHierarchyElement* FindChildWithIdentity(FHierarchyElementIdentity ChildIdentity, bool bSearchRecursively = false);
+	UE_API UHierarchyElement* FindChildWithIdentity(FHierarchyElementIdentity ChildIdentity, bool bSearchRecursively = false);
 
-	UHierarchyElement* CopyAndAddItemAsChild(const UHierarchyElement& ItemToCopy);
-	UHierarchyElement* CopyAndAddItemUnderParentIdentity(const UHierarchyElement& ItemToCopy, FHierarchyElementIdentity ParentIdentity);
+	UE_API UHierarchyElement* CopyAndAddItemAsChild(const UHierarchyElement& ItemToCopy);
+	UE_API UHierarchyElement* CopyAndAddItemUnderParentIdentity(const UHierarchyElement& ItemToCopy, FHierarchyElementIdentity ParentIdentity);
 	
 	/** Remove a child with a given identity. Can be searched recursively. This function operates under the assumption there will be only one item with a given identity. */
-	bool RemoveChildWithIdentity(FHierarchyElementIdentity ChildIdentity, bool bSearchRecursively = false);
+	UE_API bool RemoveChildWithIdentity(FHierarchyElementIdentity ChildIdentity, bool bSearchRecursively = false);
 	
 	template<class ChildClass>
 	bool DoesOneChildExist(bool bRecursive = false) const;
@@ -179,13 +181,13 @@ public:
 	/** An identity can be optionally set to create a mapping from previously existing guids or names to hierarchy items that represent them. */
 	void SetIdentity(FHierarchyElementIdentity InIdentity) { Identity = InIdentity; }
 	FHierarchyElementIdentity GetPersistentIdentity() const { return Identity; }
-	TArray<FHierarchyElementIdentity> GetParentIdentities() const;
+	UE_API TArray<FHierarchyElementIdentity> GetParentIdentities() const;
 	
 	/** Overridden modify method to also mark all children as modified */
-	virtual bool Modify(bool bAlwaysMarkDirty = true) override;
+	UE_API virtual bool Modify(bool bAlwaysMarkDirty = true) override;
 
 protected:
-	virtual void PostLoad() override;
+	UE_API virtual void PostLoad() override;
 
 	UPROPERTY()
 	TArray<TObjectPtr<UHierarchyElement>> Children;
@@ -269,8 +271,8 @@ void UHierarchyElement::SortChildren(const PREDICATE_CLASS& Predicate, bool bRec
 }
 
 /** A minimal implementation of a section. */
-UCLASS()
-class DATAHIERARCHYEDITOR_API UHierarchySection : public UHierarchyElement
+UCLASS(MinimalAPI)
+class UHierarchySection : public UHierarchyElement
 {
 	GENERATED_BODY()
 
@@ -280,7 +282,7 @@ public:
 	void SetSectionName(FName InSectionName) { Section = InSectionName; }
 	FName GetSectionName() const { return Section; }
 	
-	void SetSectionNameAsText(const FText& Text);
+	UE_API void SetSectionNameAsText(const FText& Text);
 	FText GetSectionNameAsText() const { return FText::FromName(Section); }
 
 	void SetTooltip(const FText& InTooltip) { Tooltip = InTooltip; }
@@ -297,8 +299,8 @@ private:
 };
 
 /** UHierarchyRoot is used as the main object for serialization purposes, and a transient root is created automatically by the widget to populate the source list of items. */
-UCLASS()
-class DATAHIERARCHYEDITOR_API UHierarchyRoot : public UHierarchyElement
+UCLASS(MinimalAPI)
+class UHierarchyRoot : public UHierarchyElement
 {
 	GENERATED_BODY()
 public:
@@ -308,20 +310,20 @@ public:
 	const TArray<UHierarchySection*>& GetSectionData() const { return Sections; }
 	TArray<TObjectPtr<UHierarchySection>>& GetSectionDataMutable() { return Sections; }
 
-	TSet<FName> GetSections() const;
-	int32 GetSectionIndex(UHierarchySection* Section) const;
+	UE_API TSet<FName> GetSections() const;
+	UE_API int32 GetSectionIndex(UHierarchySection* Section) const;
 
-	UHierarchySection* AddSection(FText InNewSectionName, int32 InsertIndex = INDEX_NONE, TSubclassOf<UHierarchySection> SectionClass = UHierarchySection::StaticClass());
-	UHierarchySection* FindSectionByIdentity(FHierarchyElementIdentity SectionIdentity);
+	UE_API UHierarchySection* AddSection(FText InNewSectionName, int32 InsertIndex = INDEX_NONE, TSubclassOf<UHierarchySection> SectionClass = UHierarchySection::StaticClass());
+	UE_API UHierarchySection* FindSectionByIdentity(FHierarchyElementIdentity SectionIdentity);
 	/** This will copy the section element itself */
-	void DuplicateSectionFromOtherRoot(const UHierarchySection& SectionToCopy);
-	void RemoveSection(FText SectionName);
-	void RemoveSectionByIdentity(FHierarchyElementIdentity SectionIdentity);
+	UE_API void DuplicateSectionFromOtherRoot(const UHierarchySection& SectionToCopy);
+	UE_API void RemoveSection(FText SectionName);
+	UE_API void RemoveSectionByIdentity(FHierarchyElementIdentity SectionIdentity);
 	
-	virtual bool Modify(bool bAlwaysMarkDirty = true) override;
-	virtual void Serialize(FStructuredArchive::FRecord Record) override;
+	UE_API virtual bool Modify(bool bAlwaysMarkDirty = true) override;
+	UE_API virtual void Serialize(FStructuredArchive::FRecord Record) override;
 
-	virtual void EmptyAllData();
+	UE_API virtual void EmptyAllData();
 protected:
 
 	UPROPERTY()
@@ -340,8 +342,8 @@ public:
 
 /** A category, potentially pointing at the section it belongs to. Only top-level categories can belong to sections by default.
  *  Inherit from this to add your own properties.  */
-UCLASS()
-class DATAHIERARCHYEDITOR_API UHierarchyCategory : public UHierarchyElement
+UCLASS(MinimalAPI)
+class UHierarchyCategory : public UHierarchyElement
 {
 	GENERATED_BODY()
 public:
@@ -363,12 +365,12 @@ public:
 	 *  We fix this up by looking through our available sections and match up via persistent identity.
 	 *  This function expects the correct section with the same identity to exist already at the root level
 	 */
-	void FixupSectionLinkage();
+	UE_API void FixupSectionLinkage();
 
-	static FHierarchyElementIdentity ConstructIdentity();
+	static UE_API FHierarchyElementIdentity ConstructIdentity();
 
 protected:
-	virtual void PostLoad() override;
+	UE_API virtual void PostLoad() override;
 private:
 	UPROPERTY()
 	FName Category;
@@ -385,15 +387,15 @@ struct FHierarchyElementViewModel;
 struct FHierarchySectionViewModel;
 
 /** Inherit from this to allow UI customization for your drag & drop operation by overriding CreateCustomDecorator. */
-class DATAHIERARCHYEDITOR_API FHierarchyDragDropOp : public FDragDropOperation
+class FHierarchyDragDropOp : public FDragDropOperation
 {
 public:
 	DRAG_DROP_OPERATOR_TYPE(FHierarchyDragDropOp, FDragDropOperation)
 
-	FHierarchyDragDropOp(TSharedPtr<FHierarchyElementViewModel> InDraggedElementViewModel);
+	UE_API FHierarchyDragDropOp(TSharedPtr<FHierarchyElementViewModel> InDraggedElementViewModel);
 
 	virtual void Construct() override { FDragDropOperation::Construct(); }
-	virtual TSharedPtr<SWidget> GetDefaultDecorator() const override final;
+	UE_API virtual TSharedPtr<SWidget> GetDefaultDecorator() const override final;
 	
 	/** Override this custom decorator function to provide custom widget visuals. If not specified, you can still use Label and Description. */
 	virtual TSharedRef<SWidget> CreateCustomDecorator() const { return SNullWidget::NullWidget; }
@@ -419,8 +421,8 @@ protected:
 	bool bFromSourceList = false;
 };
 
-UCLASS(BlueprintType)
-class DATAHIERARCHYEDITOR_API UHierarchyMenuContext : public UObject
+UCLASS(MinimalAPI, BlueprintType)
+class UHierarchyMenuContext : public UObject
 {
 	GENERATED_BODY()
 
@@ -431,8 +433,8 @@ public:
 
 /** The main controller class for the SHierarchyEditor widget. Defines core hierarchy rules.
  *  Inherit from this and override the required virtual functions, instantiate an object, Initialize it and pass it to the SHierarchyEditor widget. */
-UCLASS(Abstract, Transient)
-class DATAHIERARCHYEDITOR_API UDataHierarchyViewModelBase : public UObject, public FSelfRegisteringEditorUndoClient
+UCLASS(MinimalAPI, Abstract, Transient)
+class UDataHierarchyViewModelBase : public UObject, public FSelfRegisteringEditorUndoClient
 {
 public:
 	DECLARE_MULTICAST_DELEGATE(FOnHierarchyChanged)
@@ -445,26 +447,26 @@ public:
 
 	GENERATED_BODY()
 
-	UDataHierarchyViewModelBase();
-	virtual ~UDataHierarchyViewModelBase() override;
+	UE_API UDataHierarchyViewModelBase();
+	UE_API virtual ~UDataHierarchyViewModelBase() override;
 
 	/** Initialize is called automatically for you, but it is recommended to call it manually after creating the HierarchyViewModel in your own Initialize function. This lets you access external data. */
-	void Initialize();
+	UE_API void Initialize();
 	/** Call Finalize manually when you no longer need the HierarchyViewModel. */
-	void Finalize();
+	UE_API void Finalize();
 	
 	bool IsInitialized() const { return bIsInitialized; }
 	bool IsFinalized() const { return bIsFinalized; }
 	bool IsValid() const { return IsInitialized() && !IsFinalized(); }
 
-	FName GetContextMenuName () const;
+	UE_API FName GetContextMenuName () const;
 
 	/** Creates view model hierarchy elements. To create custom view models, override CreateCustomViewModelForElement. */
-	TSharedPtr<FHierarchyElementViewModel> CreateViewModelForElement(UHierarchyElement* Element, TSharedPtr<FHierarchyElementViewModel> Parent);
+	UE_API TSharedPtr<FHierarchyElementViewModel> CreateViewModelForElement(UHierarchyElement* Element, TSharedPtr<FHierarchyElementViewModel> Parent);
 	/** Get the root view model associated with the hierarchy. */
 	TSharedPtr<struct FHierarchyRootViewModel> GetHierarchyRootViewModel() const { return HierarchyRootViewModel; }
 	/** Hierarchy items reflect the already edited hierarchy. This should generally be constructed from persistent serialized data. */
-	const TArray<TSharedPtr<FHierarchyElementViewModel>>& GetHierarchyItems() const;
+	UE_API const TArray<TSharedPtr<FHierarchyElementViewModel>>& GetHierarchyItems() const;
 
 	TSharedPtr<struct FHierarchySectionViewModel> GetDefaultHierarchySectionViewModel() const { return DefaultHierarchySectionViewModel; }
 public:
@@ -476,11 +478,11 @@ public:
 	 * The root view model is also given as a way to forcefully sync view models to access additional functionality, if needed */
 	virtual void PrepareSourceItems(UHierarchyRoot* SourceRoot, TSharedPtr<FHierarchyRootViewModel> SourceRootViewModel) PURE_VIRTUAL(UDataHierarchyViewModelBase::PrepareSourceItems,);
 	/** The class used for creating categories. You can subclass UHierarchyCategory to add new properties. */
-	virtual TSubclassOf<UHierarchyCategory> GetCategoryDataClass() const;
+	UE_API virtual TSubclassOf<UHierarchyCategory> GetCategoryDataClass() const;
 	/** The class used for creating sections. You can subclass UHierarchySection to add new properties. */
-	virtual TSubclassOf<UHierarchySection> GetSectionDataClass() const;
+	UE_API virtual TSubclassOf<UHierarchySection> GetSectionDataClass() const;
 	/** Function to implement drag drop ops. FHierarchyDragDropOp is a default implementation for a single hierarchy element. Inherit from it and override CreateCustomDecorator for custom UI. */
-	virtual TSharedRef<FHierarchyDragDropOp> CreateDragDropOp(TSharedRef<FHierarchyElementViewModel> Item);
+	UE_API virtual TSharedRef<FHierarchyDragDropOp> CreateDragDropOp(TSharedRef<FHierarchyElementViewModel> Item);
 	/** This needs to return true if you want the details panel to show up. */
 	virtual bool SupportsDetailsPanel() { return true; }
 	/** Overriding this allows to define details panel instance customizations for specific UClasses */
@@ -498,38 +500,38 @@ private:
 	virtual TSharedPtr<FHierarchyElementViewModel> CreateCustomViewModelForElement(UHierarchyElement* Element, TSharedPtr<FHierarchyElementViewModel> Parent) { return nullptr; }
 	/** Override this to add custom context menu options. Defaults to common actions such as finding, renaming and deleting.
 	 * The hierarchy elements to generate this menu for are contained in a UHierarchyMenuContext object accessible through the ToolMenu. */
-	virtual void GenerateDynamicContextMenuInternal(UToolMenu* DynamicToolMenu) const;
+	UE_API virtual void GenerateDynamicContextMenuInternal(UToolMenu* DynamicToolMenu) const;
 public:
 	const UHierarchyDataRefreshContext* GetRefreshContext() const { return RefreshContext; }
 	/** Set the refresh context to easily allow Hierarchy Elements representing external data to access whether the external data still exists. */
 	void SetRefreshContext(UHierarchyDataRefreshContext* InContext) { RefreshContext = InContext; }
 	
-	UHierarchyElement* AddElementUnderRoot(TSubclassOf<UHierarchyElement> NewChildClass, FHierarchyElementIdentity ChildIdentity);
-	void AddCategory(TSharedPtr<FHierarchyElementViewModel> CategoryParent = nullptr) const;
-	void AddSection() const;
+	UE_API UHierarchyElement* AddElementUnderRoot(TSubclassOf<UHierarchyElement> NewChildClass, FHierarchyElementIdentity ChildIdentity);
+	UE_API void AddCategory(TSharedPtr<FHierarchyElementViewModel> CategoryParent = nullptr) const;
+	UE_API void AddSection() const;
 
 	/** Delete all specified elements. */
-	void DeleteElements(TArray<TSharedPtr<FHierarchyElementViewModel>> ViewModels) const;
+	UE_API void DeleteElements(TArray<TSharedPtr<FHierarchyElementViewModel>> ViewModels) const;
 	/** Special case for deleting a specific element based on its identity.
 	 *  Useful for externally removing an element from the hierarchy when you don't have access to the view model.*/
-	void DeleteElementWithIdentity(FHierarchyElementIdentity Identity);
+	UE_API void DeleteElementWithIdentity(FHierarchyElementIdentity Identity);
 
-	void NavigateToElementInHierarchy(const FHierarchyElementIdentity& HierarchyIdentity) const;
-	void NavigateToElementInHierarchy(const TSharedRef<FHierarchyElementViewModel> HierarchyElement) const;
+	UE_API void NavigateToElementInHierarchy(const FHierarchyElementIdentity& HierarchyIdentity) const;
+	UE_API void NavigateToElementInHierarchy(const TSharedRef<FHierarchyElementViewModel> HierarchyElement) const;
 	
 	/** Refreshes all data and widgets */
-	void ForceFullRefresh();
-	void ForceFullRefreshOnTimer();
-	void RequestFullRefreshNextFrame();
+	UE_API void ForceFullRefresh();
+	UE_API void ForceFullRefreshOnTimer();
+	UE_API void RequestFullRefreshNextFrame();
 	
 	TSharedRef<FUICommandList> GetCommands() const { return Commands.ToSharedRef(); }
 	
-	void OnGetChildren(TSharedPtr<FHierarchyElementViewModel> Element, TArray<TSharedPtr<FHierarchyElementViewModel>>& OutChildren) const;
+	UE_API void OnGetChildren(TSharedPtr<FHierarchyElementViewModel> Element, TArray<TSharedPtr<FHierarchyElementViewModel>>& OutChildren) const;
 	
-	void RefreshAllViews(bool bFullRefresh = false) const;
-	void RefreshSourceView(bool bFullRefresh = false) const;
-	void RefreshHierarchyView(bool bFullRefresh = false) const;
-	void RefreshSectionsView() const;
+	UE_API void RefreshAllViews(bool bFullRefresh = false) const;
+	UE_API void RefreshSourceView(bool bFullRefresh = false) const;
+	UE_API void RefreshHierarchyView(bool bFullRefresh = false) const;
+	UE_API void RefreshSectionsView() const;
 	
 	// Delegate that call functions from SHierarchyEditor
 	FSimpleDelegate& OnRefreshSourceItemsRequested() { return RefreshSourceItemsRequestedDelegate; }
@@ -547,32 +549,32 @@ public:
 	FSimpleDelegate& OnInitialized() { return OnInitializedDelegate; }
 	
 	// Sections
-	void SetActiveHierarchySection(TSharedPtr<struct FHierarchySectionViewModel>);
-	TSharedPtr<FHierarchySectionViewModel> GetActiveHierarchySection() const;
-	UHierarchySection* GetActiveHierarchySectionData() const;
-	bool IsHierarchySectionActive(const UHierarchySection* Section) const;
+	UE_API void SetActiveHierarchySection(TSharedPtr<struct FHierarchySectionViewModel>);
+	UE_API TSharedPtr<FHierarchySectionViewModel> GetActiveHierarchySection() const;
+	UE_API UHierarchySection* GetActiveHierarchySectionData() const;
+	UE_API bool IsHierarchySectionActive(const UHierarchySection* Section) const;
 	FOnSectionActivated& OnHierarchySectionActivated() { return OnHierarchySectionActivatedDelegate; }
 
-	FString OnElementToStringDebug(TSharedPtr<FHierarchyElementViewModel> ElementViewModel) const;
+	UE_API FString OnElementToStringDebug(TSharedPtr<FHierarchyElementViewModel> ElementViewModel) const;
 
 protected:	
-	virtual void PostUndo(bool bSuccess) override;
-	virtual void PostRedo(bool bSuccess) override;
-	virtual bool MatchesContext(const FTransactionContext& InContext, const TArray<TPair<UObject*, FTransactionObjectEvent>>& TransactionObjectContexts) const override;
+	UE_API virtual void PostUndo(bool bSuccess) override;
+	UE_API virtual void PostRedo(bool bSuccess) override;
+	UE_API virtual bool MatchesContext(const FTransactionContext& InContext, const TArray<TPair<UObject*, FTransactionObjectEvent>>& TransactionObjectContexts) const override;
 private:
-	bool FilterForHierarchySection(TSharedPtr<const FHierarchyElementViewModel> ViewModel) const;
-	bool FilterForUncategorizedRootItemsInAllSection(TSharedPtr<const FHierarchyElementViewModel> ViewModel) const;
+	UE_API bool FilterForHierarchySection(TSharedPtr<const FHierarchyElementViewModel> ViewModel) const;
+	UE_API bool FilterForUncategorizedRootItemsInAllSection(TSharedPtr<const FHierarchyElementViewModel> ViewModel) const;
 
-	static void GenerateDynamicContextMenu(UToolMenu* ToolMenu);
+	static UE_API void GenerateDynamicContextMenu(UToolMenu* ToolMenu);
 
-	void ToolMenuRequestRename(const FToolMenuContext& Context) const;
-	bool ToolMenuCanRequestRename(const FToolMenuContext& Context) const;
+	UE_API void ToolMenuRequestRename(const FToolMenuContext& Context) const;
+	UE_API bool ToolMenuCanRequestRename(const FToolMenuContext& Context) const;
 
-	void ToolMenuDelete(const FToolMenuContext& Context) const;
-	bool ToolMenuCanDelete(const FToolMenuContext& Context) const;
+	UE_API void ToolMenuDelete(const FToolMenuContext& Context) const;
+	UE_API bool ToolMenuCanDelete(const FToolMenuContext& Context) const;
 
-	void ToolMenuNavigateTo(const FToolMenuContext& Context) const;
-	bool ToolMenuCanNavigateTo(const FToolMenuContext& Context) const;
+	UE_API void ToolMenuNavigateTo(const FToolMenuContext& Context) const;
+	UE_API bool ToolMenuCanNavigateTo(const FToolMenuContext& Context) const;
 protected:
 	UPROPERTY()
 	TObjectPtr<UHierarchyRoot> HierarchyRoot;
@@ -618,7 +620,7 @@ private:
 /** The base view model for all elements in the hierarchy. There are four base view models inheriting from this; for roots, items, categories, and sections.
  *  When creating a new view model, you should inherit from one of those four base view models.
  */
-struct DATAHIERARCHYEDITOR_API FHierarchyElementViewModel : TSharedFromThis<FHierarchyElementViewModel>, public FTickableEditorObject
+struct FHierarchyElementViewModel : TSharedFromThis<FHierarchyElementViewModel>, public FTickableEditorObject
 {
 	DECLARE_DELEGATE(FOnSynced)
 	DECLARE_DELEGATE_RetVal_OneParam(bool, FOnFilterChild, const TSharedPtr<const FHierarchyElementViewModel> Child);
@@ -636,7 +638,7 @@ struct DATAHIERARCHYEDITOR_API FHierarchyElementViewModel : TSharedFromThis<FHie
 	/** Can be implemented for additional logic that the constructor isn't valid for. */
 	virtual void Initialize() {}
 	
-	virtual ~FHierarchyElementViewModel() override;
+	UE_API virtual ~FHierarchyElementViewModel() override;
 	
 	UHierarchyElement* GetDataMutable() const { return Element; }
 	const UHierarchyElement* GetData() const { return Element; }
@@ -647,53 +649,53 @@ struct DATAHIERARCHYEDITOR_API FHierarchyElementViewModel : TSharedFromThis<FHie
 	template<class T>
 	const T* GetData() const { return Cast<T>(Element); }
 	
-	virtual void Tick(float DeltaTime) override;
-	virtual TStatId GetStatId() const override;
+	UE_API virtual void Tick(float DeltaTime) override;
+	UE_API virtual TStatId GetStatId() const override;
 	
 	virtual FString ToString() const { return Element->ToString(); }
 	FText ToStringAsText() const { return FText::FromString(ToString()); }
 	virtual TArray<FString> GetSearchTerms() const { return {ToString()} ;}
 
-	void RefreshChildrenData();
-	void SyncViewModelsToData();
+	UE_API void RefreshChildrenData();
+	UE_API void SyncViewModelsToData();
 	
 	/** Every item view model can define its own sort order for its children. By default we put categories above items. */
-	virtual void SortChildrenData() const;
+	UE_API virtual void SortChildrenData() const;
 	
 	const TArray<TSharedPtr<FHierarchyElementViewModel>>& GetChildren() const { return Children; }
 	TArray<TSharedPtr<FHierarchyElementViewModel>>& GetChildrenMutable() { return Children; }
-	const TArray<TSharedPtr<FHierarchyElementViewModel>>& GetFilteredChildren() const;
+	UE_API const TArray<TSharedPtr<FHierarchyElementViewModel>>& GetFilteredChildren() const;
 	
-	void AddChildFilter(FOnFilterChild InFilterChild);
+	UE_API void AddChildFilter(FOnFilterChild InFilterChild);
 	
 	template<class DataClass, class ViewModelChildClass>
 	void GetChildrenViewModelsForType(TArray<TSharedPtr<ViewModelChildClass>>& OutChildren, bool bRecursive = false);
 
 	/** Returns the hierarchy depth via number of parents above. */
-	int32 GetHierarchyDepth() const;
-	bool HasParent(TSharedPtr<FHierarchyElementViewModel> ParentCandidate, bool bRecursive = false) const;
+	UE_API int32 GetHierarchyDepth() const;
+	UE_API bool HasParent(TSharedPtr<FHierarchyElementViewModel> ParentCandidate, bool bRecursive = false) const;
 	
-	TSharedRef<FHierarchyElementViewModel> DuplicateToThis(TSharedPtr<FHierarchyElementViewModel> ItemToDuplicate, int32 InsertIndex = INDEX_NONE);
-	TSharedRef<FHierarchyElementViewModel> ReparentToThis(TSharedPtr<FHierarchyElementViewModel> ItemToMove, int32 InsertIndex = INDEX_NONE);
+	UE_API TSharedRef<FHierarchyElementViewModel> DuplicateToThis(TSharedPtr<FHierarchyElementViewModel> ItemToDuplicate, int32 InsertIndex = INDEX_NONE);
+	UE_API TSharedRef<FHierarchyElementViewModel> ReparentToThis(TSharedPtr<FHierarchyElementViewModel> ItemToMove, int32 InsertIndex = INDEX_NONE);
 
-	TSharedPtr<FHierarchyElementViewModel> FindViewModelForChild(UHierarchyElement* Child, bool bSearchRecursively = false) const;
-	TSharedPtr<FHierarchyElementViewModel> FindViewModelForChild(FHierarchyElementIdentity ChildIdentity, bool bSearchRecursively = false) const;
-	int32 FindIndexOfChild(UHierarchyElement* Child) const;
-	int32 FindIndexOfDataChild(TSharedPtr<FHierarchyElementViewModel> Child) const;
-	int32 FindIndexOfDataChild(UHierarchyElement* Child) const;
+	UE_API TSharedPtr<FHierarchyElementViewModel> FindViewModelForChild(UHierarchyElement* Child, bool bSearchRecursively = false) const;
+	UE_API TSharedPtr<FHierarchyElementViewModel> FindViewModelForChild(FHierarchyElementIdentity ChildIdentity, bool bSearchRecursively = false) const;
+	UE_API int32 FindIndexOfChild(UHierarchyElement* Child) const;
+	UE_API int32 FindIndexOfDataChild(TSharedPtr<FHierarchyElementViewModel> Child) const;
+	UE_API int32 FindIndexOfDataChild(UHierarchyElement* Child) const;
 
-	UHierarchyElement* AddChild(TSubclassOf<UHierarchyElement> NewChildClass, FHierarchyElementIdentity ChildIdentity);
+	UE_API UHierarchyElement* AddChild(TSubclassOf<UHierarchyElement> NewChildClass, FHierarchyElementIdentity ChildIdentity);
 	
 	/** Deleting will ask the parent to delete its child */
-	void Delete();
-	void DeleteChild(TSharedPtr<FHierarchyElementViewModel> Child);
+	UE_API void Delete();
+	UE_API void DeleteChild(TSharedPtr<FHierarchyElementViewModel> Child);
 
 	TWeakObjectPtr<UDataHierarchyViewModelBase> GetHierarchyViewModel() const { return HierarchyViewModel; }
 
 	/** Returns a set result if the item can accept a drop either above/onto/below the item.  */
-	TOptional<EItemDropZone> OnCanRowAcceptDrop(const FDragDropEvent& DragDropEvent, EItemDropZone ItemDropZone, TSharedPtr<FHierarchyElementViewModel> Item);
-	virtual FReply OnDroppedOnRow(const FDragDropEvent& DragDropEvent, EItemDropZone ItemDropZone, TSharedPtr<FHierarchyElementViewModel> Item);
-	void OnRowDragLeave(const FDragDropEvent& DragDropEvent);
+	UE_API TOptional<EItemDropZone> OnCanRowAcceptDrop(const FDragDropEvent& DragDropEvent, EItemDropZone ItemDropZone, TSharedPtr<FHierarchyElementViewModel> Item);
+	UE_API virtual FReply OnDroppedOnRow(const FDragDropEvent& DragDropEvent, EItemDropZone ItemDropZone, TSharedPtr<FHierarchyElementViewModel> Item);
+	UE_API void OnRowDragLeave(const FDragDropEvent& DragDropEvent);
 
 	struct FCanPerformActionResults
 	{
@@ -721,7 +723,7 @@ struct DATAHIERARCHYEDITOR_API FHierarchyElementViewModel : TSharedFromThis<FHie
 	virtual bool CanHaveChildren() const { return false; }
 	
 	/** Should return true if an item should be draggable. An uneditable item can not be dragged even if CanDragInternal returns true. */
-	FCanPerformActionResults CanDrag();
+	UE_API FCanPerformActionResults CanDrag();
 	
 	/** Returns true if renamable */
 	bool CanRename() { return IsEditableByUser().bCanPerform && CanRenameInternal(); }
@@ -775,7 +777,7 @@ struct DATAHIERARCHYEDITOR_API FHierarchyElementViewModel : TSharedFromThis<FHie
 	/** Override this to register dynamic context menu entries when right clicking a single hierarchy item */
 	virtual void AppendDynamicContextMenuForSingleElement(UToolMenu* ToolMenu) {}
 	
-	FReply OnDragDetected(const FGeometry& Geometry, const FPointerEvent& PointerEvent, bool bIsSource);
+	UE_API FReply OnDragDetected(const FGeometry& Geometry, const FPointerEvent& PointerEvent, bool bIsSource);
 
 	FSimpleDelegate& GetOnRequestRename() { return OnRequestRenameDelegate; }
 	FOnSynced& GetOnSynced() { return OnSyncedDelegate; }
@@ -796,7 +798,7 @@ protected:
 	virtual bool CanDeleteInternal() { return IsForHierarchy(); }
 
 	/** Should return true if the given drag drop operation is allowed to succeed. */
-	virtual FCanPerformActionResults CanDropOnInternal(TSharedPtr<FHierarchyElementViewModel>, EItemDropZone ItemDropZone);
+	UE_API virtual FCanPerformActionResults CanDropOnInternal(TSharedPtr<FHierarchyElementViewModel>, EItemDropZone ItemDropZone);
 	
 	/** Override this to handle drop-on logic. This is called when an item has been dropped on the item that has implemented this function. */
 	virtual void OnDroppedOnInternal(TSharedPtr<FHierarchyElementViewModel>, EItemDropZone ItemDropZone) { }
@@ -815,7 +817,7 @@ private:
 	 * An example for this is when a section gets deleted, it iterates over all categories to null out the associated section  */
 	virtual void FinalizeInternal() {}
 
-	void PropagateOnChildSynced();
+	UE_API void PropagateOnChildSynced();
 protected:
 	/** The hierarchy element this view model represents. Assumed valid while this view model exists. */
 	UHierarchyElement* const Element;
@@ -857,24 +859,24 @@ void FHierarchyElementViewModel::GetChildrenViewModelsForType(TArray<TSharedPtr<
 	}
 }
 
-struct DATAHIERARCHYEDITOR_API FHierarchyRootViewModel : FHierarchyElementViewModel
+struct FHierarchyRootViewModel : FHierarchyElementViewModel
 {
 	DECLARE_DELEGATE(FOnSyncPropagated)
 	DECLARE_DELEGATE(FOnSectionsChanged)
 	DECLARE_DELEGATE_OneParam(FOnSingleSectionChanged, TSharedPtr<FHierarchySectionViewModel> AddedSection)
 	
 	FHierarchyRootViewModel(UHierarchyElement* InItem, TWeakObjectPtr<UDataHierarchyViewModelBase> InHierarchyViewModel, bool bInIsForHierarchy) : FHierarchyElementViewModel(InItem, nullptr, InHierarchyViewModel, bInIsForHierarchy) {}
-	virtual ~FHierarchyRootViewModel() override;
+	UE_API virtual ~FHierarchyRootViewModel() override;
 
-	virtual void Initialize() override;
+	UE_API virtual void Initialize() override;
 	
-	virtual FCanPerformActionResults CanDropOnInternal(TSharedPtr<FHierarchyElementViewModel> DraggedElement, EItemDropZone ItemDropZone) override;
-	virtual void OnDroppedOnInternal(TSharedPtr<FHierarchyElementViewModel> DroppedItem, EItemDropZone ItemDropZone) override;
+	UE_API virtual FCanPerformActionResults CanDropOnInternal(TSharedPtr<FHierarchyElementViewModel> DraggedElement, EItemDropZone ItemDropZone) override;
+	UE_API virtual void OnDroppedOnInternal(TSharedPtr<FHierarchyElementViewModel> DroppedItem, EItemDropZone ItemDropZone) override;
 
 	virtual bool CanHaveChildren() const override { return true; }
 
-	TSharedPtr<struct FHierarchySectionViewModel> AddSection();
-	void DeleteSection(TSharedPtr<FHierarchyElementViewModel> SectionViewModel);
+	UE_API TSharedPtr<struct FHierarchySectionViewModel> AddSection();
+	UE_API void DeleteSection(TSharedPtr<FHierarchyElementViewModel> SectionViewModel);
 	TArray<TSharedPtr<struct FHierarchySectionViewModel>>& GetSectionViewModels() { return SectionViewModels; }
 
 	FOnSyncPropagated& OnSyncPropagated() { return OnSyncPropagatedDelegate; }
@@ -886,8 +888,8 @@ struct DATAHIERARCHYEDITOR_API FHierarchyRootViewModel : FHierarchyElementViewMo
 	FOnSingleSectionChanged& OnSectionDeleted() { return OnSectionDeletedDelegate; }
 
 private:
-	void PropagateOnSynced();
-	virtual void SyncViewModelsToDataInternal() override;
+	UE_API void PropagateOnSynced();
+	UE_API virtual void SyncViewModelsToDataInternal() override;
 	TArray<TSharedPtr<struct FHierarchySectionViewModel>> SectionViewModels;
 
 	FOnSyncPropagated OnSyncPropagatedDelegate;
@@ -896,7 +898,7 @@ private:
 	FOnSectionsChanged OnSectionsChangedDelegate;
 };
 
-struct DATAHIERARCHYEDITOR_API FHierarchySectionViewModel : FHierarchyElementViewModel
+struct FHierarchySectionViewModel : FHierarchyElementViewModel
 {
 	FHierarchySectionViewModel(UHierarchySection* InItem, TSharedRef<FHierarchyRootViewModel> InParent, TWeakObjectPtr<UDataHierarchyViewModelBase> InHierarchyViewModel) : FHierarchyElementViewModel(InItem, InParent, InHierarchyViewModel, InParent->IsForHierarchy())
 	{
@@ -908,14 +910,14 @@ struct DATAHIERARCHYEDITOR_API FHierarchySectionViewModel : FHierarchyElementVie
 	
 	virtual ~FHierarchySectionViewModel() override {}
 
-	virtual FString ToString() const override;
+	UE_API virtual FString ToString() const override;
 	
-	void SetSectionName(FName InSectionName);
-	FName GetSectionName() const;
+	UE_API void SetSectionName(FName InSectionName);
+	UE_API FName GetSectionName() const;
 
-	void SetSectionNameAsText(const FText& Text);
-	FText GetSectionNameAsText() const;
-	FText GetSectionTooltip() const;
+	UE_API void SetSectionNameAsText(const FText& Text);
+	UE_API FText GetSectionNameAsText() const;
+	UE_API FText GetSectionTooltip() const;
 
 	void SetSectionImage(const FSlateBrush* InSectionImage) { SectionImage = InSectionImage; }
 	const FSlateBrush* GetSectionImage() const { return SectionImage; }
@@ -925,18 +927,18 @@ protected:
 	/** Only hierarchy sections are editable */
 	virtual FCanPerformActionResults IsEditableByUser() override { return FCanPerformActionResults(IsForHierarchy()); }
 	virtual bool CanHaveChildren() const override { return false; }
-	virtual FCanPerformActionResults CanDragInternal() override;
+	UE_API virtual FCanPerformActionResults CanDragInternal() override;
 	
 	/** We can only rename hierarchy sections */
-	virtual bool CanRenameInternal() override;
+	UE_API virtual bool CanRenameInternal() override;
 	virtual void RenameInternal(FName NewName) override { GetDataMutable<UHierarchySection>()->SetSectionName(NewName); }
 
-	virtual bool CanDeleteInternal() override;
+	UE_API virtual bool CanDeleteInternal() override;
 	
-	virtual FCanPerformActionResults CanDropOnInternal(TSharedPtr<FHierarchyElementViewModel> DraggedItem, EItemDropZone ItemDropZone) override;
-	virtual void OnDroppedOnInternal(TSharedPtr<FHierarchyElementViewModel> DroppedItem, EItemDropZone ItemDropZone) override;
+	UE_API virtual FCanPerformActionResults CanDropOnInternal(TSharedPtr<FHierarchyElementViewModel> DraggedItem, EItemDropZone ItemDropZone) override;
+	UE_API virtual void OnDroppedOnInternal(TSharedPtr<FHierarchyElementViewModel> DroppedItem, EItemDropZone ItemDropZone) override;
 
-	virtual void FinalizeInternal() override;
+	UE_API virtual void FinalizeInternal() override;
 
 private:
 	const FSlateBrush* SectionImage = nullptr;
@@ -944,7 +946,7 @@ private:
 	bool bDropDisallowed = false;
 };
 
-struct DATAHIERARCHYEDITOR_API FHierarchyItemViewModel : FHierarchyElementViewModel
+struct FHierarchyItemViewModel : FHierarchyElementViewModel
 {
 	FHierarchyItemViewModel(UHierarchyItem* InElement, TSharedRef<FHierarchyElementViewModel> InParent, TWeakObjectPtr<UDataHierarchyViewModelBase> InHierarchyViewModel) : FHierarchyElementViewModel(InElement, InParent, InHierarchyViewModel, InParent->IsForHierarchy()) {}
 	
@@ -954,11 +956,11 @@ struct DATAHIERARCHYEDITOR_API FHierarchyItemViewModel : FHierarchyElementViewMo
 	virtual bool CanHaveChildren() const override { return false; }
 	virtual FCanPerformActionResults CanDragInternal() override { return true; }
 
-	virtual FCanPerformActionResults CanDropOnInternal(TSharedPtr<FHierarchyElementViewModel> DraggedItem, EItemDropZone ItemDropZone) override;
-	virtual void OnDroppedOnInternal(TSharedPtr<FHierarchyElementViewModel> DroppedItem, EItemDropZone ItemDropZone) override;
+	UE_API virtual FCanPerformActionResults CanDropOnInternal(TSharedPtr<FHierarchyElementViewModel> DraggedItem, EItemDropZone ItemDropZone) override;
+	UE_API virtual void OnDroppedOnInternal(TSharedPtr<FHierarchyElementViewModel> DroppedItem, EItemDropZone ItemDropZone) override;
 };
 
-struct DATAHIERARCHYEDITOR_API FHierarchyCategoryViewModel : FHierarchyElementViewModel
+struct FHierarchyCategoryViewModel : FHierarchyElementViewModel
 {
 	FHierarchyCategoryViewModel(UHierarchyCategory* InCategory, TSharedRef<FHierarchyElementViewModel> InParent, TWeakObjectPtr<UDataHierarchyViewModelBase> InHierarchyViewModel) : FHierarchyElementViewModel(InCategory, InParent, InHierarchyViewModel, InParent->IsForHierarchy()) {}
 	virtual ~FHierarchyCategoryViewModel() override{}
@@ -972,13 +974,13 @@ struct DATAHIERARCHYEDITOR_API FHierarchyCategoryViewModel : FHierarchyElementVi
 	virtual void RenameInternal(FName NewName) override { GetDataMutable<UHierarchyCategory>()->SetCategoryName(NewName); }
 	virtual const UHierarchySection* GetSectionInternal() const override { return Cast<UHierarchyCategory>(Element)->GetSection(); }
 
-	bool IsTopCategoryActive() const;
+	UE_API bool IsTopCategoryActive() const;
 
-	virtual FCanPerformActionResults CanDropOnInternal(TSharedPtr<FHierarchyElementViewModel>, EItemDropZone ItemDropZone) override;
-	virtual void OnDroppedOnInternal(TSharedPtr<FHierarchyElementViewModel> DroppedItem, EItemDropZone ItemDropZone) override;
+	UE_API virtual FCanPerformActionResults CanDropOnInternal(TSharedPtr<FHierarchyElementViewModel>, EItemDropZone ItemDropZone) override;
+	UE_API virtual void OnDroppedOnInternal(TSharedPtr<FHierarchyElementViewModel> DroppedItem, EItemDropZone ItemDropZone) override;
 };
 
-class DATAHIERARCHYEDITOR_API FSectionDragDropOp : public FHierarchyDragDropOp
+class FSectionDragDropOp : public FHierarchyDragDropOp
 {
 public:
 	DRAG_DROP_OPERATOR_TYPE(FSectionDragDropOp, FHierarchyDragDropOp)
@@ -987,5 +989,7 @@ public:
 	
 	TWeakPtr<FHierarchySectionViewModel> GetDraggedSection() const { return StaticCastSharedPtr<FHierarchySectionViewModel>(DraggedElement.Pin()); }
 private:
-	virtual TSharedRef<SWidget> CreateCustomDecorator() const override;
+	UE_API virtual TSharedRef<SWidget> CreateCustomDecorator() const override;
 };
+
+#undef UE_API
diff --git a/Engine/Source/Editor/DataLayerEditor/Public/DataLayer/DataLayerEditorSubsystem.h b/Engine/Source/Editor/DataLayerEditor/Public/DataLayer/DataLayerEditorSubsystem.h
--- a/Engine/Source/Editor/DataLayerEditor/Public/DataLayer/DataLayerEditorSubsystem.h
+++ b/Engine/Source/Editor/DataLayerEditor/Public/DataLayer/DataLayerEditorSubsystem.h
@@ -32,6 +32,8 @@
 
 #include "DataLayerEditorSubsystem.generated.h"
 
+#define UE_API DATALAYEREDITOR_API
+
 class AActor;
 class AWorldDataLayers;
 class FDataLayersBroadcast;
@@ -52,11 +54,11 @@ template<typename TItemType> class IFilter;
 namespace WorldPartitionTests { class FExternalDataLayerTest; }
 
 USTRUCT(BlueprintType)
-struct DATALAYEREDITOR_API FDataLayerCreationParameters
+struct FDataLayerCreationParameters
 {
 	GENERATED_USTRUCT_BODY()
 
-	FDataLayerCreationParameters();
+	UE_API FDataLayerCreationParameters();
 
 	// Required. Will assign the asset to the created instance.
 	UPROPERTY(EditAnywhere, BlueprintReadWrite, Category = "Data Layer")
@@ -83,27 +85,27 @@ public:
 	TArray<TSoftObjectPtr<const UDataLayerAsset>> DataLayerAssets;
 };
 
-UCLASS()
-class DATALAYEREDITOR_API UDataLayerEditorSubsystem final : public UEditorSubsystem, public IActorEditorContextClient, public FTickableGameObject
+UCLASS(MinimalAPI)
+class UDataLayerEditorSubsystem final : public UEditorSubsystem, public IActorEditorContextClient, public FTickableGameObject
 {
 	GENERATED_BODY()
 
 public:
-	UDataLayerEditorSubsystem();
+	UE_API UDataLayerEditorSubsystem();
 
-	static UDataLayerEditorSubsystem* Get();
+	static UE_API UDataLayerEditorSubsystem* Get();
 
 	typedef IFilter<const TWeakObjectPtr<AActor>&> FActorFilter;
 
 	/**
 	 *	Prepares for use
 	 */
-	virtual void Initialize(FSubsystemCollectionBase& Collection) override;
+	UE_API virtual void Initialize(FSubsystemCollectionBase& Collection) override;
 
 	/**
 	 *	Internal cleanup
 	 */
-	virtual void Deinitialize() override;
+	UE_API virtual void Deinitialize() override;
 
 	/**
 	 *	Destructor
@@ -111,37 +113,37 @@ public:
 	virtual ~UDataLayerEditorSubsystem() {}
 
 	//~ Begin FTickableGameObject interface
-	virtual UWorld* GetTickableGameObjectWorld() const override;
+	UE_API virtual UWorld* GetTickableGameObjectWorld() const override;
 	virtual bool IsTickableInEditor() const { return true; }
-	virtual ETickableTickType GetTickableTickType() const override;
-	virtual bool IsTickable() const override;
-	virtual void Tick(float DeltaTime) override;
+	UE_API virtual ETickableTickType GetTickableTickType() const override;
+	UE_API virtual bool IsTickable() const override;
+	UE_API virtual void Tick(float DeltaTime) override;
 	TStatId GetStatId() const override { RETURN_QUICK_DECLARE_CYCLE_STAT(UDataLayerEditorSubsystem, STATGROUP_Tickables); }
 	//~ End FTickableGameObject interface
 
 	//~ Begin UObject overrides.
-	virtual void BeginDestroy() override;
+	UE_API virtual void BeginDestroy() override;
 	// ~End UObject overrides.
 
 	//~ Begin IActorEditorContextClient interface
-	virtual void OnExecuteActorEditorContextAction(UWorld* InWorld, const EActorEditorContextAction& InType, AActor* InActor = nullptr) override;
-	virtual void CaptureActorEditorContextState(UWorld* InWorld, UActorEditorContextStateCollection* InStateCollection) const override;
-	virtual void RestoreActorEditorContextState(UWorld* InWorld, const UActorEditorContextStateCollection* InStateCollection) override;
-	virtual bool GetActorEditorContextDisplayInfo(UWorld* InWorld, FActorEditorContextClientDisplayInfo& OutDiplayInfo) const override;
+	UE_API virtual void OnExecuteActorEditorContextAction(UWorld* InWorld, const EActorEditorContextAction& InType, AActor* InActor = nullptr) override;
+	UE_API virtual void CaptureActorEditorContextState(UWorld* InWorld, UActorEditorContextStateCollection* InStateCollection) const override;
+	UE_API virtual void RestoreActorEditorContextState(UWorld* InWorld, const UActorEditorContextStateCollection* InStateCollection) override;
+	UE_API virtual bool GetActorEditorContextDisplayInfo(UWorld* InWorld, FActorEditorContextClientDisplayInfo& OutDiplayInfo) const override;
 	virtual bool CanResetContext(UWorld* InWorld) const override { return true; };
-	virtual TSharedRef<SWidget> GetActorEditorContextWidget(UWorld* InWorld) const override;
+	UE_API virtual TSharedRef<SWidget> GetActorEditorContextWidget(UWorld* InWorld) const override;
 	virtual FOnActorEditorContextClientChanged& GetOnActorEditorContextClientChanged() override { return ActorEditorContextClientChanged; }
 	//~ End IActorEditorContextClient interface
 	UFUNCTION(BlueprintCallable, Category = DataLayers)
-	void AddToActorEditorContext(UDataLayerInstance* InDataLayerInstance);
+	UE_API void AddToActorEditorContext(UDataLayerInstance* InDataLayerInstance);
 	UFUNCTION(BlueprintCallable, Category = DataLayers)
-	void RemoveFromActorEditorContext(UDataLayerInstance* InDataLayerInstance);
+	UE_API void RemoveFromActorEditorContext(UDataLayerInstance* InDataLayerInstance);
 	UFUNCTION(BlueprintCallable, Category = DataLayers)
-	bool SetActorEditorContextCurrentExternalDataLayer(const UExternalDataLayerAsset* InExternalDataLayerAsset);
+	UE_API bool SetActorEditorContextCurrentExternalDataLayer(const UExternalDataLayerAsset* InExternalDataLayerAsset);
 	UFUNCTION(BlueprintCallable, Category = DataLayers)
-	const UExternalDataLayerAsset* GetActorEditorContextCurrentExternalDataLayer() const;
+	UE_API const UExternalDataLayerAsset* GetActorEditorContextCurrentExternalDataLayer() const;
 
-	TArray<const UDataLayerInstance*> GetDataLayerInstances(const TArray<const UDataLayerAsset*> DataLayerAssets) const;
+	UE_API TArray<const UDataLayerInstance*> GetDataLayerInstances(const TArray<const UDataLayerAsset*> DataLayerAssets) const;
 
 	template<class DataLayerType, typename ...Args>
 	DataLayerType* CreateDataLayerInstance(AWorldDataLayers* WorldDataLayers, Args&&... InArgs);
@@ -177,7 +179,7 @@ public:
 	 */
 	UE_DEPRECATED(5.4, "Use IsActorValidForDataLayerInstances instead.")
 	UFUNCTION(BlueprintCallable, Category = DataLayers)
-	bool IsActorValidForDataLayer(AActor* Actor);
+	UE_API bool IsActorValidForDataLayer(AActor* Actor);
 
 	/**
 	 *	Checks to see if the specified actor is in an appropriate state to interact with DataLayers
@@ -186,7 +188,7 @@ public:
 	 *  @param  DataLayerInstances	The data layers used to do the validation
 	 */
 	UFUNCTION(BlueprintCallable, Category = DataLayers)
-	bool IsActorValidForDataLayerInstances(AActor* Actor, const TArray<UDataLayerInstance*>& DataLayerInstances);
+	UE_API bool IsActorValidForDataLayerInstances(AActor* Actor, const TArray<UDataLayerInstance*>& DataLayerInstances);
 
 	/**
 	 * Adds the actor to the DataLayer.
@@ -196,7 +198,7 @@ public:
 	 * @return				true if the actor was added.  false is returned if the actor already belongs to the DataLayer.
 	 */
 	UFUNCTION(BlueprintCallable, Category = DataLayers)
-	bool AddActorToDataLayer(AActor* Actor, UDataLayerInstance* DataLayer);
+	UE_API bool AddActorToDataLayer(AActor* Actor, UDataLayerInstance* DataLayer);
 
 	/**
 	 * Adds the provided actor to the DataLayers.
@@ -206,7 +208,7 @@ public:
 	 * @return				true if the actor was added to at least one of the provided DataLayers.
 	 */
 	UFUNCTION(BlueprintCallable, Category = DataLayers)
-	bool AddActorToDataLayers(AActor* Actor, const TArray<UDataLayerInstance*>& DataLayers);
+	UE_API bool AddActorToDataLayers(AActor* Actor, const TArray<UDataLayerInstance*>& DataLayers);
 
 	/**
 	 * Removes an actor from the specified DataLayer.
@@ -216,7 +218,7 @@ public:
 	 * @return					true if the actor was removed from the DataLayer.  false is returned if the actor already belonged to the DataLayer.
 	 */
 	UFUNCTION(BlueprintCallable, Category = DataLayers)
-	bool RemoveActorFromDataLayer(AActor* Actor, UDataLayerInstance* DataLayerToRemove);
+	UE_API bool RemoveActorFromDataLayer(AActor* Actor, UDataLayerInstance* DataLayerToRemove);
 
 	/**
 	 * Removes the provided actor from the DataLayers.
@@ -226,7 +228,7 @@ public:
 	 * @return				true if the actor was removed from at least one of the provided DataLayers.
 	 */
 	UFUNCTION(BlueprintCallable, Category = DataLayers)
-	bool RemoveActorFromDataLayers(AActor* Actor, const TArray<UDataLayerInstance*>& DataLayers);
+	UE_API bool RemoveActorFromDataLayers(AActor* Actor, const TArray<UDataLayerInstance*>& DataLayers);
 
 	/**
 	 * Removes an actor from all DataLayers.
@@ -235,7 +237,7 @@ public:
 	 * @return					true if the actor was changed.
 	 */
 	UFUNCTION(BlueprintCallable, Category = DataLayers)
-	bool RemoveActorFromAllDataLayers(AActor* Actor);
+	UE_API bool RemoveActorFromAllDataLayers(AActor* Actor);
 	
 	/**
 	 * Removes an actor from all DataLayers.
@@ -244,7 +246,7 @@ public:
 	 * @return					true if any actor was changed.
 	 */
 	UFUNCTION(BlueprintCallable, Category = DataLayers)
-	bool RemoveActorsFromAllDataLayers(const TArray<AActor*>& Actors);
+	UE_API bool RemoveActorsFromAllDataLayers(const TArray<AActor*>& Actors);
 
 	/////////////////////////////////////////////////
 	// Operations on multiple actors.
@@ -257,7 +259,7 @@ public:
 	 * @return				true if at least one actor was added to the DataLayer.  false is returned if all the actors already belonged to the DataLayer.
 	 */
 	UFUNCTION(BlueprintCallable, Category = DataLayers)
-	bool AddActorsToDataLayer(const TArray<AActor*>& Actors, UDataLayerInstance* DataLayer);
+	UE_API bool AddActorsToDataLayer(const TArray<AActor*>& Actors, UDataLayerInstance* DataLayer);
 
 	/**
 	 * Add the actors to the DataLayers
@@ -267,7 +269,7 @@ public:
 	 * @return				true if at least one actor was added to at least one DataLayer.  false is returned if all the actors already belonged to all specified DataLayers.
 	 */
 	UFUNCTION(BlueprintCallable, Category = DataLayers)
-	bool AddActorsToDataLayers(const TArray<AActor*>& Actors, const TArray<UDataLayerInstance*>& DataLayers);
+	UE_API bool AddActorsToDataLayers(const TArray<AActor*>& Actors, const TArray<UDataLayerInstance*>& DataLayers);
 
 	/**
 	 * Removes the actors from the specified DataLayer.
@@ -277,7 +279,7 @@ public:
 	 * @return					true if at least one actor was removed from the DataLayer.  false is returned if all the actors already belonged to the DataLayer.
 	 */
 	UFUNCTION(BlueprintCallable, Category = DataLayers)
-	bool RemoveActorsFromDataLayer(const TArray<AActor*>& Actors, UDataLayerInstance* DataLayer);
+	UE_API bool RemoveActorsFromDataLayer(const TArray<AActor*>& Actors, UDataLayerInstance* DataLayer);
 
 	/**
 	 * Remove the actors from the DataLayers
@@ -287,7 +289,7 @@ public:
 	 * @return				true if at least one actor was removed from at least one DataLayer.  false is returned if none of the actors belonged to any of the specified DataLayers.
 	 */
 	UFUNCTION(BlueprintCallable, Category = DataLayers)
-	bool RemoveActorsFromDataLayers(const TArray<AActor*>& Actors, const TArray<UDataLayerInstance*>& DataLayers);
+	UE_API bool RemoveActorsFromDataLayers(const TArray<AActor*>& Actors, const TArray<UDataLayerInstance*>& DataLayers);
 
 	/////////////////////////////////////////////////
 	// Operations on selected actors.
@@ -299,7 +301,7 @@ public:
 	 * @return				true if at least one actor was added.  false is returned if all selected actors already belong to the DataLayer.
 	 */
 	UFUNCTION(BlueprintCallable, Category = DataLayers)
-	bool AddSelectedActorsToDataLayer(UDataLayerInstance* DataLayer);
+	UE_API bool AddSelectedActorsToDataLayer(UDataLayerInstance* DataLayer);
 
 	/**
 	 * Adds selected actors to the DataLayers.
@@ -308,7 +310,7 @@ public:
 	 * @return				true if at least one actor was added.  false is returned if all selected actors already belong to the DataLayers.
 	 */
 	UFUNCTION(BlueprintCallable, Category = DataLayers)
-	bool AddSelectedActorsToDataLayers(const TArray<UDataLayerInstance*>& DataLayers);
+	UE_API bool AddSelectedActorsToDataLayers(const TArray<UDataLayerInstance*>& DataLayers);
 
 	/**
 	 * Removes the selected actors from the DataLayer.
@@ -317,7 +319,7 @@ public:
 	 * @return				true if at least one actor was added.  false is returned if all selected actors already belong to the DataLayer.
 	 */
 	UFUNCTION(BlueprintCallable, Category = DataLayers)
-	bool RemoveSelectedActorsFromDataLayer(UDataLayerInstance* DataLayer);
+	UE_API bool RemoveSelectedActorsFromDataLayer(UDataLayerInstance* DataLayer);
 
 	/**
 	 * Removes selected actors from the DataLayers.
@@ -326,7 +328,7 @@ public:
 	 * @return				true if at least one actor was removed.
 	 */
 	UFUNCTION(BlueprintCallable, Category = DataLayers)
-	bool RemoveSelectedActorsFromDataLayers(const TArray<UDataLayerInstance*>& DataLayers);
+	UE_API bool RemoveSelectedActorsFromDataLayers(const TArray<UDataLayerInstance*>& DataLayers);
 
 
 	/////////////////////////////////////////////////
@@ -342,7 +344,7 @@ public:
 	 * @return									true if at least one actor was selected/deselected.
 	 */
 	UFUNCTION(BlueprintCallable, Category = DataLayers)
-	bool SelectActorsInDataLayer(UDataLayerInstance* DataLayer, const bool bSelect, const bool bNotify, const bool bSelectEvenIfHidden = false);
+	UE_API bool SelectActorsInDataLayer(UDataLayerInstance* DataLayer, const bool bSelect, const bool bNotify, const bool bSelectEvenIfHidden = false);
 
 	/**
 	 * Selects/de-selects actors belonging to the DataLayer.
@@ -354,7 +356,7 @@ public:
 	 * @param	Filter	[optional]				Actor that don't pass the specified filter restrictions won't be selected.
 	 * @return									true if at least one actor was selected/deselected.
 	 */
-	bool SelectActorsInDataLayer(UDataLayerInstance* DataLayer, const bool bSelect, const bool bNotify, const bool bSelectEvenIfHidden, const TSharedPtr<FActorFilter>& Filter);
+	UE_API bool SelectActorsInDataLayer(UDataLayerInstance* DataLayer, const bool bSelect, const bool bNotify, const bool bSelectEvenIfHidden, const TSharedPtr<FActorFilter>& Filter);
 
 	/**
 	 * Selects/de-selects actors belonging to the DataLayers.
@@ -366,7 +368,7 @@ public:
 	 * @return									true if at least one actor was selected/deselected.
 	 */
 	UFUNCTION(BlueprintCallable, Category = DataLayers)
-	bool SelectActorsInDataLayers(const TArray<UDataLayerInstance*>& DataLayers, const bool bSelect, const bool bNotify, const bool bSelectEvenIfHidden = false);
+	UE_API bool SelectActorsInDataLayers(const TArray<UDataLayerInstance*>& DataLayers, const bool bSelect, const bool bNotify, const bool bSelectEvenIfHidden = false);
 
 	/**
 	 * Selects/de-selects actors belonging to the DataLayers.
@@ -378,7 +380,7 @@ public:
 	 * @param	Filter	[optional]				Actor that don't pass the specified filter restrictions won't be selected.
 	 * @return									true if at least one actor was selected/deselected.
 	 */
-	bool SelectActorsInDataLayers(const TArray<UDataLayerInstance*>& DataLayers, const bool bSelect, const bool bNotify, const bool bSelectEvenIfHidden, const TSharedPtr<FActorFilter>& Filter);
+	UE_API bool SelectActorsInDataLayers(const TArray<UDataLayerInstance*>& DataLayers, const bool bSelect, const bool bNotify, const bool bSelectEvenIfHidden, const TSharedPtr<FActorFilter>& Filter);
 
 	/** 
 	 * Pin/unpin actors belonging to the DataLayers.
@@ -387,7 +389,7 @@ public:
 	 * @param	bPinned							If true actors are pinned; if false, actors are unpinned.
 	 * 
 	 */
-	void SetActorsPinStateInDataLayers(const TArray<UDataLayerInstance*>& DataLayerInstances, const bool bPinned);
+	UE_API void SetActorsPinStateInDataLayers(const TArray<UDataLayerInstance*>& DataLayerInstances, const bool bPinned);
 
 
 	/////////////////////////////////////////////////
@@ -405,7 +407,7 @@ public:
 	 * @param	bRedrawViewports			If true the viewports will be redrawn; if false, they will not
 	 */
 	UFUNCTION(BlueprintCallable, Category = DataLayers)
-	bool UpdateActorVisibility(AActor* Actor, bool& bOutSelectionChanged, bool& bOutActorModified, const bool bNotifySelectionChange, const bool bRedrawViewports);
+	UE_API bool UpdateActorVisibility(AActor* Actor, bool& bOutSelectionChanged, bool& bOutActorModified, const bool bNotifySelectionChange, const bool bRedrawViewports);
 
 	/**
 	 * Updates the visibility of all actors in the viewports
@@ -414,7 +416,7 @@ public:
 	 * @param	bRedrawViewports			If true the viewports will be redrawn; if false, they will not
 	 */
 	UFUNCTION(BlueprintCallable, Category = DataLayers)
-	bool UpdateAllActorsVisibility(const bool bNotifySelectionChange, const bool bRedrawViewports);
+	UE_API bool UpdateAllActorsVisibility(const bool bNotifySelectionChange, const bool bRedrawViewports);
 
 
 	/////////////////////////////////////////////////
@@ -427,7 +429,7 @@ public:
 	 *	@param	InOutActors			The list to append the found actors to.
 	 */
 	UFUNCTION(BlueprintCallable, Category = DataLayers)
-	void AppendActorsFromDataLayer(UDataLayerInstance* DataLayer, TArray<AActor*>& InOutActors) const;
+	UE_API void AppendActorsFromDataLayer(UDataLayerInstance* DataLayer, TArray<AActor*>& InOutActors) const;
 
 	/**
 	 *	Appends all the actors associated with the specified DataLayer.
@@ -436,7 +438,7 @@ public:
 	 *	@param	InOutActors			The list to append the found actors to.
 	 *  @param	Filter	[optional]	Actor that don't pass the specified filter restrictions won't be selected.
 	 */
-	void AppendActorsFromDataLayer(UDataLayerInstance* DataLayer, TArray<AActor*>& InOutActors, const TSharedPtr<FActorFilter>& Filter) const;
+	UE_API void AppendActorsFromDataLayer(UDataLayerInstance* DataLayer, TArray<AActor*>& InOutActors, const TSharedPtr<FActorFilter>& Filter) const;
 	
 	/**
 	 *	Appends all the actors associated with ANY of the specified DataLayers.
@@ -445,7 +447,7 @@ public:
 	 *	@param	InOutActors			The list to append the found actors to.
 	 */
 	UFUNCTION(BlueprintCallable, Category = DataLayers)
-	void AppendActorsFromDataLayers(const TArray<UDataLayerInstance*>& DataLayers, TArray<AActor*>& InOutActors) const;
+	UE_API void AppendActorsFromDataLayers(const TArray<UDataLayerInstance*>& DataLayers, TArray<AActor*>& InOutActors) const;
 
 	/**
 	 *	Appends all the actors associated with ANY of the specified DataLayers.
@@ -454,7 +456,7 @@ public:
 	 *	@param	InOutActors			The list to append the found actors to.
 	 *  @param	Filter	[optional]	Actor that don't pass the specified filter restrictions won't be selected.
 	 */
-	void AppendActorsFromDataLayers(const TArray<UDataLayerInstance*>& DataLayers, TArray<AActor*>& InOutActors, const TSharedPtr<FActorFilter>& Filter) const;
+	UE_API void AppendActorsFromDataLayers(const TArray<UDataLayerInstance*>& DataLayers, TArray<AActor*>& InOutActors, const TSharedPtr<FActorFilter>& Filter) const;
 
 	/**
 	 *	Gets all the actors associated with the specified DataLayer. Analog to AppendActorsFromDataLayer but it returns rather than appends the actors.
@@ -463,7 +465,7 @@ public:
 	 *	@return						The list to assign the found actors to.
 	 */
 	UFUNCTION(BlueprintCallable, Category = DataLayers)
-	TArray<AActor*> GetActorsFromDataLayer(UDataLayerInstance* DataLayer) const;
+	UE_API TArray<AActor*> GetActorsFromDataLayer(UDataLayerInstance* DataLayer) const;
 
 	/**
 	 *	Gets all the actors associated with the specified DataLayer. Analog to AppendActorsFromDataLayer but it returns rather than appends the actors.
@@ -472,7 +474,7 @@ public:
 	 *  @param	Filter	[optional]	Actor that don't pass the specified filter restrictions won't be selected.
 	 *	@return						The list to assign the found actors to.
 	 */
-	TArray<AActor*> GetActorsFromDataLayer(UDataLayerInstance* DataLayer, const TSharedPtr<FActorFilter>& Filter) const;
+	UE_API TArray<AActor*> GetActorsFromDataLayer(UDataLayerInstance* DataLayer, const TSharedPtr<FActorFilter>& Filter) const;
 
 	/**
 	 *	Gets all the actors associated with ANY of the specified DataLayers. Analog to AppendActorsFromDataLayers but it returns rather than appends the actors.
@@ -481,7 +483,7 @@ public:
 	 *	@return						The list to assign the found actors to.
 	 */
 	UFUNCTION(BlueprintCallable, Category = DataLayers)
-	TArray<AActor*> GetActorsFromDataLayers(const TArray<UDataLayerInstance*>& DataLayers) const;
+	UE_API TArray<AActor*> GetActorsFromDataLayers(const TArray<UDataLayerInstance*>& DataLayers) const;
 
 	/**
 	 *	Gets all the actors associated with ANY of the specified DataLayers. Analog to AppendActorsFromDataLayers but it returns rather than appends the actors.
@@ -490,7 +492,7 @@ public:
 	 *  @param	Filter	[optional]	Actor that don't pass the specified filter restrictions won't be selected.
 	 *	@return						The list to assign the found actors to.
 	 */
-	TArray<AActor*> GetActorsFromDataLayers(const TArray<UDataLayerInstance*>& DataLayers, const TSharedPtr<FActorFilter>& Filter) const;
+	UE_API TArray<AActor*> GetActorsFromDataLayers(const TArray<UDataLayerInstance*>& DataLayers, const TSharedPtr<FActorFilter>& Filter) const;
 
 	/**
 	 * Changes the DataLayer's visibility to the provided state
@@ -499,7 +501,7 @@ public:
 	 * @param	bIsVisible	If true the DataLayer will be visible; if false, the DataLayer will not be visible.
 	 */
 	UFUNCTION(BlueprintCallable, Category = DataLayers)
-	void SetDataLayerVisibility(UDataLayerInstance* DataLayer, const bool bIsVisible);
+	UE_API void SetDataLayerVisibility(UDataLayerInstance* DataLayer, const bool bIsVisible);
 
 	/**
 	 * Changes visibility of the DataLayers to the provided state
@@ -508,7 +510,7 @@ public:
 	 * @param	bIsVisible	If true the DataLayers will be visible; if false, the DataLayers will not be visible
 	 */
 	UFUNCTION(BlueprintCallable, Category = DataLayers)
-	void SetDataLayersVisibility(const TArray<UDataLayerInstance*>& DataLayers, const bool bIsVisible);
+	UE_API void SetDataLayersVisibility(const TArray<UDataLayerInstance*>& DataLayers, const bool bIsVisible);
 
 	/**
 	 * Toggles the DataLayer's visibility
@@ -516,7 +518,7 @@ public:
 	 * @param DataLayer	The DataLayer to affect
 	 */
 	UFUNCTION(BlueprintCallable, Category = DataLayers)
-	void ToggleDataLayerVisibility(UDataLayerInstance* DataLayer);
+	UE_API void ToggleDataLayerVisibility(UDataLayerInstance* DataLayer);
 
 	/**
 	 * Toggles the visibility of all of the DataLayers
@@ -524,7 +526,7 @@ public:
 	 * @param	DataLayers	The DataLayers to affect
 	 */
 	UFUNCTION(BlueprintCallable, Category = DataLayers)
-	void ToggleDataLayersVisibility(const TArray<UDataLayerInstance*>& DataLayers);
+	UE_API void ToggleDataLayersVisibility(const TArray<UDataLayerInstance*>& DataLayers);
 
 	/**
 	 * Changes the DataLayer's IsLoadedInEditor flag to the provided state
@@ -536,7 +538,7 @@ public:
 	 * @param	bIsFromUserChange				If this change originates from a user change or not.
 	 */
 	UFUNCTION(BlueprintCallable, Category = DataLayers)
-	bool SetDataLayerIsLoadedInEditor(UDataLayerInstance* DataLayer, const bool bIsLoadedInEditor, const bool bIsFromUserChange);
+	UE_API bool SetDataLayerIsLoadedInEditor(UDataLayerInstance* DataLayer, const bool bIsLoadedInEditor, const bool bIsFromUserChange);
 
 	/**
 	 * Changes the IsLoadedInEditor flag of the DataLayers to the provided state
@@ -548,7 +550,7 @@ public:
 	 * @param	bIsFromUserChange				If this change originates from a user change or not.
 	 */
 	UFUNCTION(BlueprintCallable, Category = DataLayers)
-	bool SetDataLayersIsLoadedInEditor(const TArray<UDataLayerInstance*>& DataLayers, const bool bIsLoadedInEditor, const bool bIsFromUserChange);
+	UE_API bool SetDataLayersIsLoadedInEditor(const TArray<UDataLayerInstance*>& DataLayers, const bool bIsLoadedInEditor, const bool bIsFromUserChange);
 
 	/**
 	 * Toggles the DataLayer's IsLoadedInEditor flag
@@ -557,7 +559,7 @@ public:
 	 * @param	bIsFromUserChange				If this change originates from a user change or not.
 	 */
 	UFUNCTION(BlueprintCallable, Category = DataLayers)
-	bool ToggleDataLayerIsLoadedInEditor(UDataLayerInstance* DataLayer, const bool bIsFromUserChange);
+	UE_API bool ToggleDataLayerIsLoadedInEditor(UDataLayerInstance* DataLayer, const bool bIsFromUserChange);
 
 	/**
 	 * Toggles the IsLoadedInEditor flag of all of the DataLayers
@@ -566,13 +568,13 @@ public:
 	 * @param	bIsFromUserChange				If this change originates from a user change or not.
 	 */
 	UFUNCTION(BlueprintCallable, Category = DataLayers)
-	bool ToggleDataLayersIsLoadedInEditor(const TArray<UDataLayerInstance*>& DataLayers, const bool bIsFromUserChange);
+	UE_API bool ToggleDataLayersIsLoadedInEditor(const TArray<UDataLayerInstance*>& DataLayers, const bool bIsFromUserChange);
 
 	/**
 	 * Set the visibility of all DataLayers to true
 	 */
 	UFUNCTION(BlueprintCallable, Category = DataLayers)
-	void MakeAllDataLayersVisible();
+	UE_API void MakeAllDataLayersVisible();
 
 	/**
 	 * Gets the UDataLayerInstance associated to the DataLayerAsset
@@ -581,7 +583,7 @@ public:
 	 * @return					The UDataLayerInstance of the provided DataLayerAsset
 	 */
 	UFUNCTION(BlueprintCallable, Category = DataLayers)
-	UDataLayerInstance* GetDataLayerInstance(const UDataLayerAsset* DataLayerAsset) const;
+	UE_API UDataLayerInstance* GetDataLayerInstance(const UDataLayerAsset* DataLayerAsset) const;
 
 	/**
 	 * Gets the UDataLayerInstances associated to the each DataLayerAssets
@@ -590,7 +592,7 @@ public:
 	 * @return					The array of UDataLayerInstances corresponding to a DataLayerAsset in the DataLayerAssets array
 	 */
 	UFUNCTION(BlueprintCallable, Category = DataLayers)
-	TArray<UDataLayerInstance*> GetDataLayerInstances(const TArray<UDataLayerAsset*> DataLayerAssets) const;
+	UE_API TArray<UDataLayerInstance*> GetDataLayerInstances(const TArray<UDataLayerAsset*> DataLayerAssets) const;
 
 	/**
 	 * Gets the UDataLayerInstance Object of the DataLayer name
@@ -598,14 +600,14 @@ public:
 	 * @param	DataLayerName	The name of the DataLayer whose UDataLayerInstance Object is returned
 	 * @return					The UDataLayerInstance Object of the provided DataLayer name
 	 */
-	UDataLayerInstance* GetDataLayerInstance(const FName& DataLayerInstanceName) const;
+	UE_API UDataLayerInstance* GetDataLayerInstance(const FName& DataLayerInstanceName) const;
 
 	/**
 	 * Gets all known DataLayers and appends them to the provided array
 	 *
 	 * @param OutDataLayers[OUT] Output array to store all known DataLayers
 	 */
-	void AddAllDataLayersTo(TArray<TWeakObjectPtr<UDataLayerInstance>>& OutDataLayers) const;
+	UE_API void AddAllDataLayersTo(TArray<TWeakObjectPtr<UDataLayerInstance>>& OutDataLayers) const;
 
 	/**
 	 * Creates a UDataLayerInstance Object
@@ -614,7 +616,7 @@ public:
 	 * @return	The newly created UDataLayerInstance Object
 	 */
 	UFUNCTION(BlueprintCallable, Category = DataLayers)
-	UDataLayerInstance* CreateDataLayerInstance(const FDataLayerCreationParameters& Parameters);
+	UE_API UDataLayerInstance* CreateDataLayerInstance(const FDataLayerCreationParameters& Parameters);
 
 	/**
 	 * Sets a Parent DataLayer for a specified DataLayer
@@ -625,7 +627,7 @@ public:
 	 *  @return	true if succeeded, false if failed.
 	 */
 	UFUNCTION(BlueprintCallable, Category = DataLayers)
-	bool SetParentDataLayer(UDataLayerInstance* DataLayer, UDataLayerInstance* ParentDataLayer);
+	UE_API bool SetParentDataLayer(UDataLayerInstance* DataLayer, UDataLayerInstance* ParentDataLayer);
 
 	/**
 	 * Sets a Parent DataLayer for a specified list of DataLayers
@@ -634,7 +636,7 @@ public:
 	 *  @param ParentDataLayer	The parent DataLayer.
 	 */
 	UFUNCTION(BlueprintCallable, Category = DataLayers)
-	void SetParentDataLayerForDataLayers(const TArray<UDataLayerInstance*>& DataLayers, UDataLayerInstance* ParentDataLayer);
+	UE_API void SetParentDataLayerForDataLayers(const TArray<UDataLayerInstance*>& DataLayers, UDataLayerInstance* ParentDataLayer);
 
 	/**
 	 * Sets the initial runtime state of the DataLayers
@@ -643,7 +645,7 @@ public:
 	 * @param	InitialRuntimeState				The initial runtime state.
 	 */
 	UFUNCTION(BlueprintCallable, Category = DataLayers)
-	void SetDataLayerInitialRuntimeState(UDataLayerInstance* DataLayer, EDataLayerRuntimeState InitialRuntimeState);
+	UE_API void SetDataLayerInitialRuntimeState(UDataLayerInstance* DataLayer, EDataLayerRuntimeState InitialRuntimeState);
 
 	/**
 	 * Sets the initial editor visibility state of the DataLayers
@@ -652,7 +654,7 @@ public:
 	 * @param	InitialRuntimeState				The initial editor visibility state.
 	 */
 	UFUNCTION(BlueprintCallable, Category = DataLayers)
-	void SetDataLayerIsInitiallyVisible(UDataLayerInstance* DataLayer, bool bIsInitiallyVisible);
+	UE_API void SetDataLayerIsInitiallyVisible(UDataLayerInstance* DataLayer, bool bIsInitiallyVisible);
 
 	/**
 	 * Deletes all of the provided DataLayers
@@ -660,7 +662,7 @@ public:
 	 * @param DataLayersToDelete	A valid list of DataLayer.
 	 */
 	UFUNCTION(BlueprintCallable, Category = DataLayers)
-	void DeleteDataLayers(const TArray<UDataLayerInstance*>& DataLayersToDelete);
+	UE_API void DeleteDataLayers(const TArray<UDataLayerInstance*>& DataLayersToDelete);
 
 	/**
 	 * Deletes the provided DataLayer
@@ -668,18 +670,18 @@ public:
 	 * @param DataLayerToDelete		A valid DataLayer
 	 */
 	UFUNCTION(BlueprintCallable, Category = DataLayers)
-	void DeleteDataLayer(UDataLayerInstance* DataLayerToDelete);
+	UE_API void DeleteDataLayer(UDataLayerInstance* DataLayerToDelete);
 
 	/**
 	 * Returns all Data Layers
 	 */
 	UFUNCTION(BlueprintCallable, Category = DataLayers)
-	TArray<UDataLayerInstance*> GetAllDataLayers();
+	UE_API TArray<UDataLayerInstance*> GetAllDataLayers();
 
 	/**
 	* Resets user override settings of all DataLayers
 	*/
-	bool ResetUserSettings();
+	UE_API bool ResetUserSettings();
 
 	/*
 	 * Returns whether the DataLayer contains one of more actors that are part of the editor selection.
@@ -689,19 +691,19 @@ public:
 	/**
 	* Whether there are any deprecated DataLayerInstance found
 	*/
-	bool HasDeprecatedDataLayers() const;
+	UE_API bool HasDeprecatedDataLayers() const;
 
 	/**
 	 * Assign new unique short name to DataLayerInstance if it supports it. 
 	 */
-	bool SetDataLayerShortName(UDataLayerInstance* InDataLayerInstance, const FString& InNewShortName);
+	UE_API bool SetDataLayerShortName(UDataLayerInstance* InDataLayerInstance, const FString& InNewShortName);
 
 	/**
 	 * Move actors to the Actor Editor Context's Data Layers.
 	 * 
 	 * @return true if operation succeeded.
 	 */
-	bool ApplyActorEditorContextDataLayersToActors(const TArray<AActor*>& InActors);
+	UE_API bool ApplyActorEditorContextDataLayersToActors(const TArray<AActor*>& InActors);
 
 	//~ Begin Deprecated
 
@@ -741,15 +743,15 @@ public:
 			
 	UE_DEPRECATED(5.1, "Use GetDataLayerInstance instead")
 	UFUNCTION(BlueprintCallable, Category = DataLayers)
-	UDataLayerInstance* GetDataLayer(const FActorDataLayer& ActorDataLayer) const;
+	UE_API UDataLayerInstance* GetDataLayer(const FActorDataLayer& ActorDataLayer) const;
 
 	UE_DEPRECATED(5.1, "Use GetDataLayerInstance instead")
 	UFUNCTION(BlueprintCallable, Category = DataLayers)
-	UDataLayerInstance* GetDataLayerFromLabel(const FName& DataLayerLabel) const;
+	UE_API UDataLayerInstance* GetDataLayerFromLabel(const FName& DataLayerLabel) const;
 
 	UE_DEPRECATED(5.1, "Renaming is not permitted anymore")
 	UFUNCTION(BlueprintCallable, Category = DataLayers)
-	bool RenameDataLayer(UDataLayerInstance* DataLayer, const FName& NewDataLayerLabel);
+	UE_API bool RenameDataLayer(UDataLayerInstance* DataLayer, const FName& NewDataLayerLabel);
 
 	UE_DEPRECATED(5.1, "Use CreateDataLayerInstance with FDataLayerCreationParameters instead")
 	UFUNCTION(BlueprintCallable, Category = DataLayers)
@@ -764,65 +766,65 @@ private:
 	 *
 	 *	@param	Actor	The actor to initialize
 	 */
-	void InitializeNewActorDataLayers(AActor* Actor);
+	UE_API void InitializeNewActorDataLayers(AActor* Actor);
 
 	/**
 	 * Find and return the selected actors.
 	 *
 	 * @return				The selected AActor's as a TArray.
 	 */
-	TArray<AActor*> GetSelectedActors() const;
+	UE_API TArray<AActor*> GetSelectedActors() const;
 
 	/**
 	 * Get the current UWorld object.
 	 *
 	 * @return						The UWorld* object
 	 */
-	UWorld* GetWorld() const; // Fallback to GWorld
+	UE_API UWorld* GetWorld() const; // Fallback to GWorld
 
 	/** Delegate handler for FEditorDelegates::MapChange. It internally calls DataLayerChanged.Broadcast. */
-	void EditorMapChange();
+	UE_API void EditorMapChange();
 
 	/** Delegate handler for FEditorDelegates::RefreshDataLayerBrowser. It internally calls UpdateAllActorsVisibility to refresh the actors of each DataLayer. */
-	void EditorRefreshDataLayerBrowser();
+	UE_API void EditorRefreshDataLayerBrowser();
 
 	/** Delegate handler for FEditorDelegates::PostUndoRedo. It internally calls DataLayerChanged.Broadcast and UpdateAllActorsVisibility to refresh the actors of each DataLayer. */
-	void PostUndoRedo();
+	UE_API void PostUndoRedo();
 
-	void UpdateRegisteredWorldDelegates();
-	void OnWorldPartitionInitialized(UWorldPartition* InWorldPartition);
-	void OnWorldPartitionUninitialized(UWorldPartition* InWorldPartition);
-	void OnLoadedActorAddedToLevel(AActor& InActor);
-	void OnDataLayerEditorLoadingStateChanged(bool bIsFromUserChange);
+	UE_API void UpdateRegisteredWorldDelegates();
+	UE_API void OnWorldPartitionInitialized(UWorldPartition* InWorldPartition);
+	UE_API void OnWorldPartitionUninitialized(UWorldPartition* InWorldPartition);
+	UE_API void OnLoadedActorAddedToLevel(AActor& InActor);
+	UE_API void OnDataLayerEditorLoadingStateChanged(bool bIsFromUserChange);
 
-	bool SetDataLayerIsLoadedInEditorInternal(UDataLayerInstance* DataLayer, const bool bIsLoadedInEditor, const bool bIsFromUserChange);
+	UE_API bool SetDataLayerIsLoadedInEditorInternal(UDataLayerInstance* DataLayer, const bool bIsLoadedInEditor, const bool bIsFromUserChange);
 
-	bool PassDataLayersFilter(UWorld* World, const FWorldPartitionHandle& ActorHandle);
+	UE_API bool PassDataLayersFilter(UWorld* World, const FWorldPartitionHandle& ActorHandle);
 
-	void BroadcastActorDataLayersChanged(const TWeakObjectPtr<AActor>& ChangedActor);
-	void BroadcastDataLayerChanged(const EDataLayerAction Action, const TWeakObjectPtr<const UDataLayerInstance>& ChangedDataLayer, const FName& ChangedProperty);
-	void BroadcastDataLayerEditorLoadingStateChanged(bool bIsFromUserChange);
-	void OnSelectionChanged();
-	void RebuildSelectedDataLayersFromEditorSelection();
-	const TSet<TWeakObjectPtr<const UDataLayerInstance>>& GetSelectedDataLayersFromEditorSelection() const;
-	bool SetParentDataLayerForDataLayersInternal(const TArray<UDataLayerInstance*>& DataLayers, UDataLayerInstance* ParentDataLayer);
+	UE_API void BroadcastActorDataLayersChanged(const TWeakObjectPtr<AActor>& ChangedActor);
+	UE_API void BroadcastDataLayerChanged(const EDataLayerAction Action, const TWeakObjectPtr<const UDataLayerInstance>& ChangedDataLayer, const FName& ChangedProperty);
+	UE_API void BroadcastDataLayerEditorLoadingStateChanged(bool bIsFromUserChange);
+	UE_API void OnSelectionChanged();
+	UE_API void RebuildSelectedDataLayersFromEditorSelection();
+	UE_API const TSet<TWeakObjectPtr<const UDataLayerInstance>>& GetSelectedDataLayersFromEditorSelection() const;
+	UE_API bool SetParentDataLayerForDataLayersInternal(const TArray<UDataLayerInstance*>& DataLayers, UDataLayerInstance* ParentDataLayer);
 
-	bool UpdateAllActorsVisibility(const bool bNotifySelectionChange, const bool bRedrawViewports, ULevel* InLevel);
+	UE_API bool UpdateAllActorsVisibility(const bool bNotifySelectionChange, const bool bRedrawViewports, ULevel* InLevel);
 
-	bool IsActorValidForDataLayerForClasses(AActor* Actor, const TSet<TSubclassOf<UDataLayerInstance>>& DataLayerInstanceClasses);
+	UE_API bool IsActorValidForDataLayerForClasses(AActor* Actor, const TSet<TSubclassOf<UDataLayerInstance>>& DataLayerInstanceClasses);
 
 	// External Data Layer methods
-	void OnActorPreSpawnInitialization(AActor* InActor);
-	void OnNewActorsPlaced(UObject* InObjToUse, const TArray<AActor*>& InPlacedActors);
-	void OnEditorActorReplaced(AActor* InOldActor, AActor* InNewActor);
-	void OnExternalDataLayerAssetRegistrationStateChanged(const UExternalDataLayerAsset* ExternalDataLayerAsset, EExternalDataLayerRegistrationState OldState, EExternalDataLayerRegistrationState NewState);
-	TUniquePtr<FLevelEditorDragDropWorldSurrogateReferencingObject> OnLevelEditorDragDropWorldSurrogateReferencingObject(UWorld* ReferencingWorld, const FSoftObjectPath& Object);
-	void ApplyContext(AActor* InActor, bool bInForceTryApply = false, AActor* InReplacedActor = nullptr);
-	const UExternalDataLayerInstance* GetActorSpawningExternalDataLayerInstance(AActor* InActor) const;
-	static const UExternalDataLayerAsset* GetReferencingWorldSurrogateObjectForObject(UWorld* ReferencingWorld, const FSoftObjectPath& ObjectPath);
-	bool MoveActorToDataLayers(AActor* InActor, const TArray<UDataLayerInstance*>& InDataLayerInstances);
-	void MoveActorToExternalDataLayer(AActor* InActor, const UExternalDataLayerInstance* InExternalDataLayerInstance, bool bInNotifyFailureReason = true);
-	bool ShouldHandleActor(AActor* InActor) const;
+	UE_API void OnActorPreSpawnInitialization(AActor* InActor);
+	UE_API void OnNewActorsPlaced(UObject* InObjToUse, const TArray<AActor*>& InPlacedActors);
+	UE_API void OnEditorActorReplaced(AActor* InOldActor, AActor* InNewActor);
+	UE_API void OnExternalDataLayerAssetRegistrationStateChanged(const UExternalDataLayerAsset* ExternalDataLayerAsset, EExternalDataLayerRegistrationState OldState, EExternalDataLayerRegistrationState NewState);
+	UE_API TUniquePtr<FLevelEditorDragDropWorldSurrogateReferencingObject> OnLevelEditorDragDropWorldSurrogateReferencingObject(UWorld* ReferencingWorld, const FSoftObjectPath& Object);
+	UE_API void ApplyContext(AActor* InActor, bool bInForceTryApply = false, AActor* InReplacedActor = nullptr);
+	UE_API const UExternalDataLayerInstance* GetActorSpawningExternalDataLayerInstance(AActor* InActor) const;
+	static UE_API const UExternalDataLayerAsset* GetReferencingWorldSurrogateObjectForObject(UWorld* ReferencingWorld, const FSoftObjectPath& ObjectPath);
+	UE_API bool MoveActorToDataLayers(AActor* InActor, const TArray<UDataLayerInstance*>& InDataLayerInstances);
+	UE_API void MoveActorToExternalDataLayer(AActor* InActor, const UExternalDataLayerInstance* InExternalDataLayerInstance, bool bInNotifyFailureReason = true);
+	UE_API bool ShouldHandleActor(AActor* InActor) const;
 
 	/** Contains Data Layers that contain actors that are part of the editor selection */
 	mutable TSet<TWeakObjectPtr<const UDataLayerInstance>> SelectedDataLayersFromEditorSelection;
@@ -887,4 +889,6 @@ DataLayerInstanceType* UDataLayerEditorSubsystem::CreateDataLayerInstance(AWorld
 	}
 
 	return CastChecked<DataLayerInstanceType>(NewDataLayer);
-}
\ No newline at end of file
+}
+
+#undef UE_API
diff --git a/Engine/Source/Editor/DataLayerEditor/Public/DataLayer/DataLayerPropertyTypeCustomizationHelper.h b/Engine/Source/Editor/DataLayerEditor/Public/DataLayer/DataLayerPropertyTypeCustomizationHelper.h
--- a/Engine/Source/Editor/DataLayerEditor/Public/DataLayer/DataLayerPropertyTypeCustomizationHelper.h
+++ b/Engine/Source/Editor/DataLayerEditor/Public/DataLayer/DataLayerPropertyTypeCustomizationHelper.h
@@ -7,10 +7,14 @@
 #include "Templates/SharedPointer.h"
 #include "Widgets/SWidget.h"
 
+#define UE_API DATALAYEREDITOR_API
+
 class SWidget;
 class UDataLayerInstance;
 
-struct DATALAYEREDITOR_API FDataLayerPropertyTypeCustomizationHelper
+struct FDataLayerPropertyTypeCustomizationHelper
 {
-	static TSharedRef<SWidget> CreateDataLayerMenu(TFunction<void(const UDataLayerInstance* DataLayer)> OnDataLayerSelectedFunction, TFunction<bool(const UDataLayerInstance*)> OnShouldFilterDataLayerInstanceFunction = [](const UDataLayerInstance*){ return false; });
+	static UE_API TSharedRef<SWidget> CreateDataLayerMenu(TFunction<void(const UDataLayerInstance* DataLayer)> OnDataLayerSelectedFunction, TFunction<bool(const UDataLayerInstance*)> OnShouldFilterDataLayerInstanceFunction = [](const UDataLayerInstance*){ return false; });
 };
+
+#undef UE_API
diff --git a/Engine/Source/Editor/DerivedDataEditor/Public/DerivedDataEditor.h b/Engine/Source/Editor/DerivedDataEditor/Public/DerivedDataEditor.h
--- a/Engine/Source/Editor/DerivedDataEditor/Public/DerivedDataEditor.h
+++ b/Engine/Source/Editor/DerivedDataEditor/Public/DerivedDataEditor.h
@@ -9,6 +9,8 @@
 #include "Templates/UniquePtr.h"
 #include "IDerivedDataCacheNotifications.h"
 
+#define UE_API DERIVEDDATAEDITOR_API
+
 class FSpawnTabArgs;
 class SDerivedDataCacheSettingsDialog;
 class IDerivedDataCacheNotifications;
@@ -19,32 +21,32 @@ class SWindow;
 /**
  * The module holding all of the UI related pieces for DerivedData
  */
-class DERIVEDDATAEDITOR_API FDerivedDataEditor : public IModuleInterface
+class FDerivedDataEditor : public IModuleInterface
 {
 public:
 
 	/**
 	 * Called right after the module DLL has been loaded and the module object has been created
 	 */
-	virtual void StartupModule() override;
+	UE_API virtual void StartupModule() override;
 
 	/**
 	 * Called before the module is unloaded, right before the module object is destroyed.
 	 */
-	virtual void ShutdownModule() override;
+	UE_API virtual void ShutdownModule() override;
 
-	TSharedRef<SWidget>	CreateStatusBarWidget();
+	UE_API TSharedRef<SWidget>	CreateStatusBarWidget();
 
-	void ShowResourceUsageTab();
-	void ShowCacheStatisticsTab();
+	UE_API void ShowResourceUsageTab();
+	UE_API void ShowCacheStatisticsTab();
 
 private:
 
-	TSharedPtr<SWidget> CreateResourceUsageDialog();
-	TSharedPtr<SWidget> CreateCacheStatisticsDialog();
+	UE_API TSharedPtr<SWidget> CreateResourceUsageDialog();
+	UE_API TSharedPtr<SWidget> CreateCacheStatisticsDialog();
 
-	TSharedRef<SDockTab> CreateResourceUsageTab(const FSpawnTabArgs& Args);
-	TSharedRef<SDockTab> CreateCacheStatisticsTab(const FSpawnTabArgs& Args);
+	UE_API TSharedRef<SDockTab> CreateResourceUsageTab(const FSpawnTabArgs& Args);
+	UE_API TSharedRef<SDockTab> CreateCacheStatisticsTab(const FSpawnTabArgs& Args);
 
 	TWeakPtr<SDockTab> ResourceUsageTab;
 	TWeakPtr<SDockTab> CacheStatisticsTab;
@@ -55,3 +57,4 @@ private:
 };
 
 
+#undef UE_API
diff --git a/Engine/Source/Editor/DerivedDataWidgets/Public/DerivedDataCacheNotifications.h b/Engine/Source/Editor/DerivedDataWidgets/Public/DerivedDataCacheNotifications.h
--- a/Engine/Source/Editor/DerivedDataWidgets/Public/DerivedDataCacheNotifications.h
+++ b/Engine/Source/Editor/DerivedDataWidgets/Public/DerivedDataCacheNotifications.h
@@ -6,21 +6,23 @@
 #include "IDerivedDataCacheNotifications.h"
 #include "Templates/SharedPointer.h"
 
+#define UE_API DERIVEDDATAWIDGETS_API
+
 class SNotificationItem;
 
-class DERIVEDDATAWIDGETS_API FDerivedDataCacheNotifications : public IDerivedDataCacheNotifications
+class FDerivedDataCacheNotifications : public IDerivedDataCacheNotifications
 {
 public:
-	FDerivedDataCacheNotifications();
-	virtual ~FDerivedDataCacheNotifications();
+	UE_API FDerivedDataCacheNotifications();
+	UE_API virtual ~FDerivedDataCacheNotifications();
 
 private: 
 
 	/** DDC data put notification handler */
-	void OnDDCNotificationEvent(FDerivedDataCacheInterface::EDDCNotification DDCNotification);
+	UE_API void OnDDCNotificationEvent(FDerivedDataCacheInterface::EDDCNotification DDCNotification);
 	
 	/** Subscribe to the notifications */
-	void Subscribe(bool bSubscribe);
+	UE_API void Subscribe(bool bSubscribe);
 
 	/** Whether we are subscribed or not **/
 	bool bSubscribed;
@@ -29,3 +31,4 @@ private:
 	TSharedPtr<SNotificationItem> SharedDDCNotification;
 };
 
+#undef UE_API
diff --git a/Engine/Source/Editor/DerivedDataWidgets/Public/DerivedDataInformation.h b/Engine/Source/Editor/DerivedDataWidgets/Public/DerivedDataInformation.h
--- a/Engine/Source/Editor/DerivedDataWidgets/Public/DerivedDataInformation.h
+++ b/Engine/Source/Editor/DerivedDataWidgets/Public/DerivedDataInformation.h
@@ -5,6 +5,8 @@
 #include "HAL/Platform.h"
 #include "Internationalization/Text.h"
 
+#define UE_API DERIVEDDATAWIDGETS_API
+
 enum class DERIVEDDATAWIDGETS_API ERemoteCacheState : uint8
 {
 	Idle,
@@ -13,29 +15,31 @@ enum class DERIVEDDATAWIDGETS_API ERemoteCacheState : uint8
 	Warning,
 };
 
-class DERIVEDDATAWIDGETS_API FDerivedDataInformation
+class FDerivedDataInformation
 {
 public:
 
-	static double				GetCacheActivityTimeSeconds(bool bGet, bool bLocal);
-	static double				GetCacheActivitySizeBytes(bool bGet, bool bLocal);
-	static bool					GetHasRemoteCache();
-	static bool					GetHasZenCache();
-	static bool					GetHasUnrealCloudCache();
+	static UE_API double				GetCacheActivityTimeSeconds(bool bGet, bool bLocal);
+	static UE_API double				GetCacheActivitySizeBytes(bool bGet, bool bLocal);
+	static UE_API bool					GetHasRemoteCache();
+	static UE_API bool					GetHasZenCache();
+	static UE_API bool					GetHasUnrealCloudCache();
 	static ERemoteCacheState	GetRemoteCacheState() { return RemoteCacheState; }
-	static FText				GetRemoteCacheStateAsText();
+	static UE_API FText				GetRemoteCacheStateAsText();
 	static FText				GetRemoteCacheWarningMessage() { return RemoteCacheWarningMessage; }
-	static void					UpdateRemoteCacheState();
+	static UE_API void					UpdateRemoteCacheState();
 	static bool					IsUploading() { return bIsUploading; }
 	static bool					IsDownloading() { return bIsDownloading; }
 
 private:
 
-	static ERemoteCacheState	RemoteCacheState;
-	static FText				RemoteCacheWarningMessage;
-	static double				LastGetTime;
-	static double				LastPutTime;
-	static bool					bIsUploading;
-	static bool					bIsDownloading;
+	static UE_API ERemoteCacheState	RemoteCacheState;
+	static UE_API FText				RemoteCacheWarningMessage;
+	static UE_API double				LastGetTime;
+	static UE_API double				LastPutTime;
+	static UE_API bool					bIsUploading;
+	static UE_API bool					bIsDownloading;
 
 };
+
+#undef UE_API
diff --git a/Engine/Source/Editor/DerivedDataWidgets/Public/IDerivedDataCacheNotifications.h b/Engine/Source/Editor/DerivedDataWidgets/Public/IDerivedDataCacheNotifications.h
--- a/Engine/Source/Editor/DerivedDataWidgets/Public/IDerivedDataCacheNotifications.h
+++ b/Engine/Source/Editor/DerivedDataWidgets/Public/IDerivedDataCacheNotifications.h
@@ -4,7 +4,7 @@
 
 #include "CoreMinimal.h"
 
-class DERIVEDDATAWIDGETS_API IDerivedDataCacheNotifications
+class IDerivedDataCacheNotifications
 {
 public:
 	IDerivedDataCacheNotifications()
@@ -14,4 +14,4 @@ public:
 	virtual ~IDerivedDataCacheNotifications()
 	{
 	}
-};
\ No newline at end of file
+};
diff --git a/Engine/Source/Editor/DerivedDataWidgets/Public/SDerivedDataDialogs.h b/Engine/Source/Editor/DerivedDataWidgets/Public/SDerivedDataDialogs.h
--- a/Engine/Source/Editor/DerivedDataWidgets/Public/SDerivedDataDialogs.h
+++ b/Engine/Source/Editor/DerivedDataWidgets/Public/SDerivedDataDialogs.h
@@ -8,52 +8,56 @@
 #include "Widgets/SBoxPanel.h"
 #include "Widgets/SCompoundWidget.h"
 
+#define UE_API DERIVEDDATAWIDGETS_API
+
 class SWidget;
 
-class DERIVEDDATAWIDGETS_API SDerivedDataRemoteStoreDialog : public SCompoundWidget
+class SDerivedDataRemoteStoreDialog : public SCompoundWidget
 {
 	SLATE_BEGIN_ARGS(SDerivedDataRemoteStoreDialog) {}
 	SLATE_END_ARGS()
 
-	void Construct(const FArguments& InArgs);
+	UE_API void Construct(const FArguments& InArgs);
 
 private:
 
-	TSharedRef<SWidget> GetGridPanel();
+	UE_API TSharedRef<SWidget> GetGridPanel();
 
-	EActiveTimerReturnType UpdateGridPanels(double InCurrentTime, float InDeltaTime);
+	UE_API EActiveTimerReturnType UpdateGridPanels(double InCurrentTime, float InDeltaTime);
 
 	SVerticalBox::FSlot* GridSlot = nullptr;
 };
 
-class DERIVEDDATAWIDGETS_API SDerivedDataCacheStatisticsDialog : public SCompoundWidget
+class SDerivedDataCacheStatisticsDialog : public SCompoundWidget
 {
 	SLATE_BEGIN_ARGS(SDerivedDataCacheStatisticsDialog) {}
 	SLATE_END_ARGS()
 
-	void Construct(const FArguments& InArgs);
+	UE_API void Construct(const FArguments& InArgs);
 
 private:
 
-	TSharedRef<SWidget> GetGridPanel();
+	UE_API TSharedRef<SWidget> GetGridPanel();
 	
-	EActiveTimerReturnType UpdateGridPanels(double InCurrentTime, float InDeltaTime);
+	UE_API EActiveTimerReturnType UpdateGridPanels(double InCurrentTime, float InDeltaTime);
 
 	SVerticalBox::FSlot* GridSlot = nullptr;
 };
 
-class DERIVEDDATAWIDGETS_API SDerivedDataResourceUsageDialog : public SCompoundWidget
+class SDerivedDataResourceUsageDialog : public SCompoundWidget
 {
 	SLATE_BEGIN_ARGS(SDerivedDataResourceUsageDialog) {}
 	SLATE_END_ARGS()
 
-		void Construct(const FArguments& InArgs);
+		UE_API void Construct(const FArguments& InArgs);
 
 private:
 
-	TSharedRef<SWidget> GetGridPanel();
+	UE_API TSharedRef<SWidget> GetGridPanel();
 	
-	EActiveTimerReturnType UpdateGridPanels(double InCurrentTime, float InDeltaTime);
+	UE_API EActiveTimerReturnType UpdateGridPanels(double InCurrentTime, float InDeltaTime);
 
 	SVerticalBox::FSlot* GridSlot = nullptr;
 };
+
+#undef UE_API
diff --git a/Engine/Source/Editor/DetailCustomizations/Private/BodySetupDetails.h b/Engine/Source/Editor/DetailCustomizations/Private/BodySetupDetails.h
--- a/Engine/Source/Editor/DetailCustomizations/Private/BodySetupDetails.h
+++ b/Engine/Source/Editor/DetailCustomizations/Private/BodySetupDetails.h
@@ -8,23 +8,25 @@
 #include "UObject/WeakObjectPtr.h"
 #include "UObject/WeakObjectPtrTemplates.h"
 
+#define UE_API DETAILCUSTOMIZATIONS_API
+
 class IDetailLayoutBuilder;
 class IPropertyHandle;
 class UObject;
 
-class DETAILCUSTOMIZATIONS_API FBodySetupDetails : public IDetailCustomization
+class FBodySetupDetails : public IDetailCustomization
 {
 public:
 	/** Makes a new instance of this detail layout class for a specific detail view requesting it */
-	static TSharedRef<IDetailCustomization> MakeInstance();
+	static UE_API TSharedRef<IDetailCustomization> MakeInstance();
 
 	/** IDetailCustomization interface */
-	virtual void CustomizeDetails( IDetailLayoutBuilder& DetailBuilder ) override;
+	UE_API virtual void CustomizeDetails( IDetailLayoutBuilder& DetailBuilder ) override;
 
 	virtual void CustomizeCoMNudge(IDetailLayoutBuilder& DetailBuilder, TSharedRef<IPropertyHandle> BodyInstanceHandler) {};
 
 private:
-	void AddPhysicalAnimation(IDetailLayoutBuilder& DetailBuilder);
+	UE_API void AddPhysicalAnimation(IDetailLayoutBuilder& DetailBuilder);
 
 private:
 	TSharedPtr<class FBodyInstanceCustomizationHelper> BodyInstanceCustomizationHelper;
@@ -47,3 +49,4 @@ private:
 	TArray< TWeakObjectPtr<UObject> > ObjectsCustomized;
 };
 
+#undef UE_API
diff --git a/Engine/Source/Editor/DetailCustomizations/Private/LightingChannelsCustomization.h b/Engine/Source/Editor/DetailCustomizations/Private/LightingChannelsCustomization.h
--- a/Engine/Source/Editor/DetailCustomizations/Private/LightingChannelsCustomization.h
+++ b/Engine/Source/Editor/DetailCustomizations/Private/LightingChannelsCustomization.h
@@ -5,29 +5,33 @@
 #include "IPropertyTypeCustomization.h"
 #include "PropertyHandle.h"
 
+#define UE_API DETAILCUSTOMIZATIONS_API
+
 /**
  * Customizes Lighting Channels as a horizontal row of buttons.
  */
-class DETAILCUSTOMIZATIONS_API FLightingChannelsCustomization : public IPropertyTypeCustomization
+class FLightingChannelsCustomization : public IPropertyTypeCustomization
 {
 public:
 	
-	static TSharedRef<IPropertyTypeCustomization> MakeInstance();
+	static UE_API TSharedRef<IPropertyTypeCustomization> MakeInstance();
 
 	FLightingChannelsCustomization(){}
 
-	virtual void CustomizeHeader(TSharedRef<IPropertyHandle> StructPropertyHandle, class FDetailWidgetRow& HeaderRow, IPropertyTypeCustomizationUtils& StructCustomizationUtils) override;
-	virtual void CustomizeChildren(TSharedRef<IPropertyHandle> PropertyHandle, class IDetailChildrenBuilder& ChildBuilder, IPropertyTypeCustomizationUtils& CustomizationUtils) override;
+	UE_API virtual void CustomizeHeader(TSharedRef<IPropertyHandle> StructPropertyHandle, class FDetailWidgetRow& HeaderRow, IPropertyTypeCustomizationUtils& StructCustomizationUtils) override;
+	UE_API virtual void CustomizeChildren(TSharedRef<IPropertyHandle> PropertyHandle, class IDetailChildrenBuilder& ChildBuilder, IPropertyTypeCustomizationUtils& CustomizationUtils) override;
 
 private:
 
-	FText GetStructPropertyNameText() const;
-	FText GetStructPropertyTooltipText() const;
+	UE_API FText GetStructPropertyNameText() const;
+	UE_API FText GetStructPropertyTooltipText() const;
 
-	bool IsLightingChannelButtonEditable(uint32 ChildIndex) const;
-	void OnButtonCheckedStateChanged(ECheckBoxState NewState, uint32 ChildIndex) const;
-	ECheckBoxState GetButtonCheckedState(uint32 ChildIndex) const;
+	UE_API bool IsLightingChannelButtonEditable(uint32 ChildIndex) const;
+	UE_API void OnButtonCheckedStateChanged(ECheckBoxState NewState, uint32 ChildIndex) const;
+	UE_API ECheckBoxState GetButtonCheckedState(uint32 ChildIndex) const;
 
 	TSharedPtr<IPropertyHandle> LightingChannelsHandle;
 	FCheckBoxStyle Style;
 };
+
+#undef UE_API
diff --git a/Engine/Source/Editor/DetailCustomizations/Private/SkeletalMeshComponentDetails.h b/Engine/Source/Editor/DetailCustomizations/Private/SkeletalMeshComponentDetails.h
--- a/Engine/Source/Editor/DetailCustomizations/Private/SkeletalMeshComponentDetails.h
+++ b/Engine/Source/Editor/DetailCustomizations/Private/SkeletalMeshComponentDetails.h
@@ -13,6 +13,8 @@
 #include "UObject/WeakObjectPtr.h"
 #include "UObject/WeakObjectPtrTemplates.h"
 
+#define UE_API DETAILCUSTOMIZATIONS_API
+
 class FDelegateHandle;
 class IDetailLayoutBuilder;
 class IPropertyHandle;
@@ -23,72 +25,72 @@ class UObject;
 class USkeleton;
 struct FAssetData;
 
-class DETAILCUSTOMIZATIONS_API FSkeletalMeshComponentDetails : public IDetailCustomization
+class FSkeletalMeshComponentDetails : public IDetailCustomization
 {
 public:
-	FSkeletalMeshComponentDetails();
-	~FSkeletalMeshComponentDetails();
+	UE_API FSkeletalMeshComponentDetails();
+	UE_API ~FSkeletalMeshComponentDetails();
 
 	/** Makes a new instance of this detail layout class for a specific detail view requesting it */
-	static TSharedRef<IDetailCustomization> MakeInstance();
+	static UE_API TSharedRef<IDetailCustomization> MakeInstance();
 
 	/** IDetailCustomization interface */
-	virtual void CustomizeDetails( IDetailLayoutBuilder& DetailBuilder ) override;
+	UE_API virtual void CustomizeDetails( IDetailLayoutBuilder& DetailBuilder ) override;
 
-	static TSharedRef<SWidget> CreateAsyncSceneValueWidgetWithWarning(const TSharedPtr<IPropertyHandle>& AsyncScenePropertyHandle);
+	static UE_API TSharedRef<SWidget> CreateAsyncSceneValueWidgetWithWarning(const TSharedPtr<IPropertyHandle>& AsyncScenePropertyHandle);
 
 private:
-	void UpdateAnimationCategory(IDetailLayoutBuilder& DetailBuilder);
-	void UpdatePhysicsCategory(IDetailLayoutBuilder& DetailBuilder);
+	UE_API void UpdateAnimationCategory(IDetailLayoutBuilder& DetailBuilder);
+	UE_API void UpdatePhysicsCategory(IDetailLayoutBuilder& DetailBuilder);
 
 	/** Function that returns whether the specified animation mode should be visible */
-	EVisibility VisibilityForAnimationMode(EAnimationMode::Type AnimationMode) const;
+	UE_API EVisibility VisibilityForAnimationMode(EAnimationMode::Type AnimationMode) const;
 
 	/** Helper wrapper functions for VisibilityForAnimationMode */
-	EVisibility VisibilityForBlueprintMode() const;
+	UE_API EVisibility VisibilityForBlueprintMode() const;
 	EVisibility VisibilityForSingleAnimMode() const { return VisibilityForAnimationMode(EAnimationMode::AnimationSingleNode); }
-	bool AnimPickerIsEnabled() const;
+	UE_API bool AnimPickerIsEnabled() const;
 
-	EVisibility VisibilityForAnimModeProperty() const;
+	UE_API EVisibility VisibilityForAnimModeProperty() const;
 
 	/** Handler for filtering animation assets in the UI picker when asset mode is selected */
-	bool OnShouldFilterAnimAsset(const FAssetData& AssetData);
+	UE_API bool OnShouldFilterAnimAsset(const FAssetData& AssetData);
 
 	/** Delegate called when a skeletal mesh property is changed on a selected object */
 	USkeletalMeshComponent::FOnSkeletalMeshPropertyChanged OnSkeletalMeshPropertyChanged;
 	
 	/** Register/Unregister the mesh changed delegate to TargetComponent */
-	void PerformInitialRegistrationOfSkeletalMeshes(IDetailLayoutBuilder& DetailBuilder);
-	void RegisterSkeletalMeshPropertyChanged(TWeakObjectPtr<USkeletalMeshComponent> Mesh);
-	void UnregisterSkeletalMeshPropertyChanged(TWeakObjectPtr<USkeletalMeshComponent> Mesh);
-	void UnregisterAllMeshPropertyChangedCallers();
+	UE_API void PerformInitialRegistrationOfSkeletalMeshes(IDetailLayoutBuilder& DetailBuilder);
+	UE_API void RegisterSkeletalMeshPropertyChanged(TWeakObjectPtr<USkeletalMeshComponent> Mesh);
+	UE_API void UnregisterSkeletalMeshPropertyChanged(TWeakObjectPtr<USkeletalMeshComponent> Mesh);
+	UE_API void UnregisterAllMeshPropertyChangedCallers();
 
 	/**
 	* Iterates over registered meshes and returns a pointer to the common skeleton used by all of them.
 	* If the meshes use more than one different skeleton, NULL is returned.
 	*/
-	USkeleton* GetValidSkeletonFromRegisteredMeshes() const;
+	UE_API USkeleton* GetValidSkeletonFromRegisteredMeshes() const;
 
 	/** Bound to the delegate used to detect changes in skeletal mesh properties */
-	void SkeletalMeshPropertyChanged();
+	UE_API void SkeletalMeshPropertyChanged();
 
 	/** Generates menu content for the class picker when it is clicked */
-	TSharedRef<SWidget> GetClassPickerMenuContent();
+	UE_API TSharedRef<SWidget> GetClassPickerMenuContent();
 
 	/** Gets the currently selected blueprint name to display on the class picker combo button */
-	FText GetSelectedAnimBlueprintName() const;
+	UE_API FText GetSelectedAnimBlueprintName() const;
 
 	/** Callback from the class picker when the user selects a class */
-	void OnClassPicked(UClass* PickedClass);
+	UE_API void OnClassPicked(UClass* PickedClass);
 
 	/** Callback from the detail panel to browse to the selected anim asset */
-	void OnBrowseToAnimBlueprint();
+	UE_API void OnBrowseToAnimBlueprint();
 
 	/** Callback from the details panel to use the currently selected asset in the content browser */
-	void UseSelectedAnimBlueprint();
+	UE_API void UseSelectedAnimBlueprint();
 
 	/** Called when a skeletal mesh property changes. */
-	void UpdateSkeletonNameAndPickerVisibility();
+	UE_API void UpdateSkeletonNameAndPickerVisibility();
 
 	/** Cached layout builder for use after customization */
 	IDetailLayoutBuilder* CurrentDetailBuilder;
@@ -123,3 +125,5 @@ private:
 	/** Per-mesh handles to registered OnSkeletalMeshPropertyChanged delegates */
 	TMap<USkeletalMeshComponent*, FDelegateHandle> OnSkeletalMeshPropertyChangedDelegateHandles;
 };
+
+#undef UE_API
