Data Layer - Hide Streaming Priority field on editor data layer

diff --git a/Engine/Source/Editor/DataLayerEditor/Private/DataLayer/DataLayerInstanceCustomization.cpp b/Engine/Source/Editor/DataLayerEditor/Private/DataLayer/DataLayerInstanceCustomization.cpp
--- a/Engine/Source/Editor/DataLayerEditor/Private/DataLayer/DataLayerInstanceCustomization.cpp
+++ b/Engine/Source/Editor/DataLayerEditor/Private/DataLayer/DataLayerInstanceCustomization.cpp
@@ -47,6 +47,7 @@ void FDataLayerInstanceDetails::CustomizeDetails(IDetailLayoutBuilder& DetailBui
 	{
 		DetailBuilder.HideProperty(GET_MEMBER_NAME_CHECKED(UDataLayerInstance, InitialRuntimeState));
 		DetailBuilder.HideProperty(GET_MEMBER_NAME_CHECKED(UDataLayerInstance, OverrideBlockOnSlowStreaming));
+		DetailBuilder.HideProperty(GET_MEMBER_NAME_CHECKED(UDataLayerInstance, StreamingPriority));
 	}
 
 	TSharedRef<IPropertyHandle> DataLayerAssetProperty = DetailBuilder.GetProperty("DataLayerAsset", UDataLayerInstanceWithAsset::StaticClass());
