[World] Add icons and mnemonic text for Data Layer context menu, and add the rename entry which was only accessible through F2 before. Required adding Input Binding text override arguments to ToolMenu APIs.

[RN] Improved UX of the Data Layers context menu
#tests Renamed layers, checked that layers that cannot be renamed had the right tooltip, checked that items still worked as expected

#Synced-CL 49997005

diff --git a/Engine/Source/Developer/ToolMenus/Private/ToolMenuEntry.cpp b/Engine/Source/Developer/ToolMenus/Private/ToolMenuEntry.cpp
--- a/Engine/Source/Developer/ToolMenus/Private/ToolMenuEntry.cpp
+++ b/Engine/Source/Developer/ToolMenus/Private/ToolMenuEntry.cpp
@@ -243,7 +243,7 @@ void FToolMenuEntry::SetCommand(const TSharedPtr<const FUICommandInfo>& InComman
 	Icon = InIcon.IsSet() ? InIcon : InCommand->GetIcon();
 }
 
-FToolMenuEntry FToolMenuEntry::InitMenuEntry(const FName InName, const TAttribute<FText>& InLabel, const TAttribute<FText>& InToolTip, const TAttribute<FSlateIcon>& InIcon, const FToolUIActionChoice& InAction, const EUserInterfaceActionType InUserInterfaceActionType, const FName InTutorialHighlightName)
+FToolMenuEntry FToolMenuEntry::InitMenuEntry(const FName InName, const TAttribute<FText>& InLabel, const TAttribute<FText>& InToolTip, const TAttribute<FSlateIcon>& InIcon, const FToolUIActionChoice& InAction, const EUserInterfaceActionType InUserInterfaceActionType, const FName InTutorialHighlightName, const TOptional<FText>& InInputBindingOverride)
 {
 	FToolMenuEntry Entry(UToolMenus::Get()->CurrentOwner(), InName, EMultiBlockType::MenuEntry);
 	Entry.TutorialHighlightName = InTutorialHighlightName;
@@ -252,6 +252,10 @@ FToolMenuEntry FToolMenuEntry::InitMenuEntry(const FName InName, const TAttribut
 	Entry.Icon = InIcon;
 	Entry.UserInterfaceActionType = InUserInterfaceActionType;
 	Entry.Action = InAction;
+	if (InInputBindingOverride.IsSet())
+	{
+		Entry.InputBindingLabel = InInputBindingOverride.GetValue();
+	}
 	return Entry;
 }
 
diff --git a/Engine/Source/Developer/ToolMenus/Private/ToolMenuSection.cpp b/Engine/Source/Developer/ToolMenus/Private/ToolMenuSection.cpp
--- a/Engine/Source/Developer/ToolMenus/Private/ToolMenuSection.cpp
+++ b/Engine/Source/Developer/ToolMenus/Private/ToolMenuSection.cpp
@@ -83,9 +83,9 @@ FToolMenuEntry& FToolMenuSection::AddEntryObject(UToolMenuEntryScript* InObject)
 	return AddEntry(Args);
 }
 
-FToolMenuEntry& FToolMenuSection::AddMenuEntry(const FName InName, const TAttribute<FText>& InLabel, const TAttribute<FText>& InToolTip, const TAttribute<FSlateIcon>& InIcon, const FToolUIActionChoice& InAction, const EUserInterfaceActionType InUserInterfaceActionType, const FName InTutorialHighlightName)
+FToolMenuEntry& FToolMenuSection::AddMenuEntry(const FName InName, const TAttribute<FText>& InLabel, const TAttribute<FText>& InToolTip, const TAttribute<FSlateIcon>& InIcon, const FToolUIActionChoice& InAction, const EUserInterfaceActionType InUserInterfaceActionType, const FName InTutorialHighlightName, const TOptional<FText>& InInputBindingOverride)
 {
-	return AddEntry(FToolMenuEntry::InitMenuEntry(InName, InLabel, InToolTip, InIcon, InAction, InUserInterfaceActionType, InTutorialHighlightName));
+	return AddEntry(FToolMenuEntry::InitMenuEntry(InName, InLabel, InToolTip, InIcon, InAction, InUserInterfaceActionType, InTutorialHighlightName, InInputBindingOverride));
 }
 
 FToolMenuEntry& FToolMenuSection::AddMenuEntry(const TSharedPtr< const FUICommandInfo >& InCommand, const TAttribute<FText>& InLabelOverride, const TAttribute<FText>& InToolTipOverride, const TAttribute<FSlateIcon>& InIconOverride, const FName InTutorialHighlightName, const TOptional<FName> InNameOverride)
diff --git a/Engine/Source/Developer/ToolMenus/Public/ToolMenuEntry.h b/Engine/Source/Developer/ToolMenus/Public/ToolMenuEntry.h
--- a/Engine/Source/Developer/ToolMenus/Public/ToolMenuEntry.h
+++ b/Engine/Source/Developer/ToolMenus/Public/ToolMenuEntry.h
@@ -175,7 +175,7 @@ struct FToolMenuEntry
 	UE_API FToolMenuEntry& operator=(const FToolMenuEntry&);
 	UE_API FToolMenuEntry& operator=(FToolMenuEntry&&);
 
-	static UE_API FToolMenuEntry InitMenuEntry(const FName InName, const TAttribute<FText>& InLabel, const TAttribute<FText>& InToolTip, const TAttribute<FSlateIcon>& InIcon, const FToolUIActionChoice& InAction, const EUserInterfaceActionType UserInterfaceActionType = EUserInterfaceActionType::Button, const FName InTutorialHighlightName = NAME_None);
+	static UE_API FToolMenuEntry InitMenuEntry(const FName InName, const TAttribute<FText>& InLabel, const TAttribute<FText>& InToolTip, const TAttribute<FSlateIcon>& InIcon, const FToolUIActionChoice& InAction, const EUserInterfaceActionType UserInterfaceActionType = EUserInterfaceActionType::Button, const FName InTutorialHighlightName = NAME_None, const TOptional<FText>& InInputBindingOverride = TOptional<FText>());
 	static UE_API FToolMenuEntry InitMenuEntry(const TSharedPtr< const FUICommandInfo >& InCommand, const TAttribute<FText>& InLabelOverride = TAttribute<FText>(), const TAttribute<FText>& InToolTipOverride = TAttribute<FText>(), const TAttribute<FSlateIcon>& InIconOverride = TAttribute<FSlateIcon>(), const FName InTutorialHighlightName = NAME_None, const TOptional<FName> InNameOverride = TOptional<FName>());
 	static UE_API FToolMenuEntry InitMenuEntry(const FName InNameOverride, const TSharedPtr< const FUICommandInfo >& InCommand, const TAttribute<FText>& InLabelOverride = TAttribute<FText>(), const TAttribute<FText>& InToolTipOverride = TAttribute<FText>(), const TAttribute<FSlateIcon>& InIconOverride = TAttribute<FSlateIcon>(), const FName InTutorialHighlightName = NAME_None);
 	static UE_API FToolMenuEntry InitMenuEntryWithCommandList(const TSharedPtr<const FUICommandInfo>& InCommand, const TSharedPtr<const FUICommandList>& InCommandList, const TAttribute<FText>& InLabelOverride = TAttribute<FText>(), const TAttribute<FText>& InToolTipOverride = TAttribute<FText>(), const TAttribute<FSlateIcon>& InIconOverride = TAttribute<FSlateIcon>(), const FName InTutorialHighlightName = NAME_None, const TOptional<FName> InNameOverride = TOptional<FName>());
@@ -222,7 +222,7 @@ struct FToolMenuEntry
 	/**
 	 * Show this menu entry in the top-level toolbar section of a toolbar.
 	 *
-	 * Entries of a toolbar submenu can be raised to the top-level of the toolbar. Such top-level entires appear in the
+	 * Entries of a toolbar submenu can be raised to the top-level of the toolbar. Such top-level entries appear in the
 	 * toolbar to the right of the submenu they belong to.
 	 *
 	 * This flag only effects entries within submenus of toolbar-type ToolMenus.
diff --git a/Engine/Source/Developer/ToolMenus/Public/ToolMenuSection.h b/Engine/Source/Developer/ToolMenus/Public/ToolMenuSection.h
--- a/Engine/Source/Developer/ToolMenus/Public/ToolMenuSection.h
+++ b/Engine/Source/Developer/ToolMenus/Public/ToolMenuSection.h
@@ -55,9 +55,11 @@ public:
 	 * @param InIcon The user-visible icon of this entry.
 	 * @param InAction The delegates to call to make this entry interactive. Supports more than just on-click delegates. See FToolUIActionChoice for more information.
 	 * @param UserInterfaceActionType The type of user interface element (such as a checkbox) to use for this entry. See EUserInterfaceActionType for more options.
+	 * @param InTutorialHighlightName The tutorial highlight key
+	 * @param InInputBindingOverride The key-binding reminder
 	 * @return A reference to the added menu entry.
 	 */
-	UE_API FToolMenuEntry& AddMenuEntry(const FName InName, const TAttribute<FText>& InLabel, const TAttribute<FText>& InToolTip, const TAttribute<FSlateIcon>& InIcon, const FToolUIActionChoice& InAction, const EUserInterfaceActionType UserInterfaceActionType = EUserInterfaceActionType::Button, const FName InTutorialHighlightName = NAME_None);
+	UE_API FToolMenuEntry& AddMenuEntry(const FName InName, const TAttribute<FText>& InLabel, const TAttribute<FText>& InToolTip, const TAttribute<FSlateIcon>& InIcon, const FToolUIActionChoice& InAction, const EUserInterfaceActionType UserInterfaceActionType = EUserInterfaceActionType::Button, const FName InTutorialHighlightName = NAME_None, const TOptional<FText>& InInputBindingOverride = TOptional<FText>());
 	UE_API FToolMenuEntry& AddMenuEntry(const TSharedPtr< const FUICommandInfo >& InCommand, const TAttribute<FText>& InLabelOverride = TAttribute<FText>(), const TAttribute<FText>& InToolTipOverride = TAttribute<FText>(), const TAttribute<FSlateIcon>& InIconOverride = TAttribute<FSlateIcon>(), const FName InTutorialHighlightName = NAME_None, const TOptional<FName> InNameOverride = TOptional<FName>());
 	UE_API FToolMenuEntry& AddMenuEntry(const FName InNameOverride, const TSharedPtr< const FUICommandInfo >& InCommand, const TAttribute<FText>& InLabelOverride = TAttribute<FText>(), const TAttribute<FText>& InToolTipOverride = TAttribute<FText>(), const TAttribute<FSlateIcon>& InIconOverride = TAttribute<FSlateIcon>(), const FName InTutorialHighlightName = NAME_None);
 	UE_API FToolMenuEntry& AddMenuEntryWithCommandList(const TSharedPtr<const FUICommandInfo>& InCommand, const TSharedPtr<const FUICommandList>& InCommandList, const TAttribute<FText>& InLabelOverride = TAttribute<FText>(), const TAttribute<FText>& InToolTipOverride = TAttribute<FText>(), const TAttribute<FSlateIcon>& InIconOverride = TAttribute<FSlateIcon>(), const FName InTutorialHighlightName = NAME_None, const TOptional<FName> InNameOverride = TOptional<FName>());
diff --git a/Engine/Source/Editor/DataLayerEditor/Private/DataLayer/DataLayerMode.cpp b/Engine/Source/Editor/DataLayerEditor/Private/DataLayer/DataLayerMode.cpp
--- a/Engine/Source/Editor/DataLayerEditor/Private/DataLayer/DataLayerMode.cpp
+++ b/Engine/Source/Editor/DataLayerEditor/Private/DataLayer/DataLayerMode.cpp
@@ -531,7 +531,7 @@ FReply FDataLayerMode::OnKeyDown(const FKeyEvent& InKeyEvent)
 {
 	const FSceneOutlinerItemSelection& Selection = SceneOutliner->GetSelection();
 
-	// Rename key: Rename selected actors (not rebindable, because it doesn't make much sense to bind.)
+	// Rename key: Rename selected data layer (not rebindable, because it doesn't make much sense to bind.)
 	if (InKeyEvent.GetKey() == EKeys::F2)
 	{
 		if (Selection.Num() == 1)
@@ -558,7 +558,6 @@ FReply FDataLayerMode::OnKeyDown(const FKeyEvent& InKeyEvent)
 	{
 		DeleteItems(Selection.SelectedItems);
 		return FReply::Handled();
-
 	}
 
 	if (Commands->ProcessCommandBindings(InKeyEvent))
@@ -1529,19 +1528,49 @@ void FDataLayerMode::RegisterContextMenu()
 								}),
 								FSimpleDelegate::CreateLambda([]() { FSlateApplication::Get().DismissAllMenus(); }));
 							Section.AddEntry(FToolMenuEntry::InitWidget("PickDataLayerAsset", MenuWidget, FText::GetEmpty(), false));
-						}));
+						}), false, FSlateIcon(FCoreStyle::Get().GetStyleSetName(), "Icons.Import"));
 				}
 				else
 				{
-					Section.AddMenuEntry("CreateNewDataLayer", LOCTEXT("CreateNewDataLayer", "Create New Data Layer"), FText(), FSlateIcon(),
+					Section.AddMenuEntry("CreateNewDataLayer", LOCTEXT("CreateNewDataLayer", "Create New Data Layer"), FText(), FSlateIcon(FCoreStyle::Get().GetStyleSetName(), "Icons.AddCircle"),
 						FUIAction(FExecuteAction::CreateLambda([CreateNewEmptyDataLayer]() { CreateNewEmptyDataLayer(); })));
 				}
 
-				Section.AddMenuEntry("CreateNewDataLayerPrivate", LOCTEXT("CreateNewDataLayerPrivate", "Create New Private Data Layer"), LOCTEXT("CreateNewDataLayerPrivateToolTip", "Creates an Editor Only Data Layer that cannot be used by other Worlds."), FSlateIcon(),
+				Section.AddMenuEntry("CreateNewDataLayerPrivate", LOCTEXT("CreateNewDataLayerPrivate", "Create New Private Data Layer"), LOCTEXT("CreateNewDataLayerPrivateToolTip", "Creates an Editor Only Data Layer that cannot be used by other Worlds."), FSlateIcon(FCoreStyle::Get().GetStyleSetName(), "Icons.AddCircle"),
 					FUIAction(FExecuteAction::CreateLambda([CreateNewEmptyDataLayer]() { CreateNewEmptyDataLayer(true); } )));
 
+				bool bCanRename = false;
+				if (SelectedDataLayers.Num() == 1)
+				{
+					FSceneOutlinerTreeItemPtr ItemToRename = SceneOutliner->GetSelection().SelectedItems[0].Pin();
+					if (ItemToRename.IsValid() && Mode->CanRenameItem(*ItemToRename) && ItemToRename->CanInteract())
+					{
+						bCanRename = true;
+					}
+				}
+
+				Section.AddMenuEntry("RenameSelectedDataLayer", LOCTEXT("RenameSelectedDataLayer", "Rename Selected Data Layer"), (bCanRename) ? FText() : LOCTEXT("RenameSelectedDataLayerTooltip", "Only Private data layers can be renamed, as Asset data layers take their name from the asset. Also a single data layer can be renamed at a time"), FSlateIcon(FCoreStyle::Get().GetStyleSetName(), "GenericCommands.Rename"),
+					FUIAction(
+						FExecuteAction::CreateLambda([Mode, SelectedDataLayers, SceneOutliner]() 
+						{
+							if (SelectedDataLayers.Num() == 1)
+							{
+								FSceneOutlinerTreeItemPtr ItemToRename = SceneOutliner->GetSelection().SelectedItems[0].Pin();
+								if (ItemToRename.IsValid() && Mode->CanRenameItem(*ItemToRename) && ItemToRename->CanInteract())
+								{
+									SceneOutliner->SetPendingRenameItem(ItemToRename);
+									SceneOutliner->ScrollItemIntoView(ItemToRename);
+								}
+							}
+						}),
+						FCanExecuteAction::CreateLambda([bCanRename]
+						{
+							return bCanRename;
+						})
+					), EUserInterfaceActionType::Button, NAME_None, LOCTEXT("F2", "F2"));
+
 				const bool bAllSelectedDataLayersCanBeRemoved = Algo::AllOf(SelectedDataLayers, [](const UDataLayerInstance* DataLayerInstance) { return DataLayerInstance->CanBeRemoved(); });
-				Section.AddMenuEntry("DeleteSelectedDataLayers", LOCTEXT("DeleteSelectedDataLayers", "Delete Selected Data Layer(s)"), FText(), FSlateIcon(),
+				Section.AddMenuEntry("DeleteSelectedDataLayers", LOCTEXT("DeleteSelectedDataLayers", "Delete Selected Data Layer(s)"), FText(), FSlateIcon(FCoreStyle::Get().GetStyleSetName(), "GenericCommands.Delete"),
 					FUIAction(
 						FExecuteAction::CreateLambda([Mode, SelectedDataLayers]() 
 						{
@@ -1552,7 +1581,7 @@ void FDataLayerMode::RegisterContextMenu()
 						{
 							return !SelectedDataLayers.IsEmpty() && !bSelectedDataLayersContainsLocked && bAllSelectedDataLayersCanBeRemoved;
 						})
-					));
+					), EUserInterfaceActionType::Button, NAME_None, LOCTEXT("Delete", "Delete"));
 
 				const bool bSelectedDataLayersSupportParenting = Algo::AllOf(SelectedDataLayers, [](const UDataLayerInstance* DataLayerInstance) { return !DataLayerInstance->IsReadOnly() && DataLayerInstance->CanHaveParentDataLayerInstance(); });
 				if (!SelectedDataLayers.IsEmpty() && !bSelectedDataLayersContainsLocked && bSelectedDataLayersSupportParenting)
@@ -1587,10 +1616,10 @@ void FDataLayerMode::RegisterContextMenu()
 									return false;
 								}),
 								bShowRoot);
-						}));
+						}), false, FSlateIcon(FCoreStyle::Get().GetStyleSetName(), "Icons.Reimport"));
 				}
 
-				Section.AddMenuEntry("Copy Selected Data Layer Instances(s) File Path", LOCTEXT("CopySelectedDataLayerInstancessFilePath", "Copy Selected Data Layer Instance(s) File Path"), FText(), FSlateIcon(),
+				Section.AddMenuEntry("Copy Selected Data Layer Instances(s) File Path", LOCTEXT("CopySelectedDataLayerInstancessFilePath", "Copy Selected Data Layer Instance(s) File Path"), FText(), FSlateIcon(FCoreStyle::Get().GetStyleSetName(), "GenericCommands.Copy"),
 					FUIAction(
 						FExecuteAction::CreateLambda([SelectedDataLayers]()
 						{
@@ -1609,7 +1638,7 @@ void FDataLayerMode::RegisterContextMenu()
 
 				const bool bAllSelectedDataLayersCanAddActors = Algo::AllOf(SelectedDataLayers, [](const UDataLayerInstance* DataLayerInstance) { return DataLayerInstance->CanUserAddActors(); });
 				FText AddActorToolTip = bSelectedActorsAreAllUserManaged ? FText::GetEmpty() : SelectionContainsNonUserManagedActorsText;
-				Section.AddMenuEntry("AddSelectedActorsToSelectedDataLayers", LOCTEXT("AddSelectedActorsToSelectedDataLayers", "Add Selected Actor(s) to Selected Data Layer(s)"), AddActorToolTip, FSlateIcon(),
+				Section.AddMenuEntry("AddSelectedActorsToSelectedDataLayers", LOCTEXT("AddSelectedActorsToSelectedDataLayers", "Add Selected Actor(s) to Selected Data Layer(s)"), AddActorToolTip, FSlateIcon(FCoreStyle::Get().GetStyleSetName(), "DataLayerBrowser.AddSelection"),
 					FUIAction(
 						FExecuteAction::CreateLambda([SelectedDataLayers]()
 						{
@@ -1655,12 +1684,12 @@ void FDataLayerMode::RegisterContextMenu()
 									return false;
 								})
 							);
-						}));
+						}), false, FSlateIcon(FCoreStyle::Get().GetStyleSetName(), "DataLayerBrowser.AddSelection"));
 				}
 
 				const bool bAllSelectedDataLayersCanRemoveActors = Algo::AllOf(SelectedDataLayers, [](const UDataLayerInstance* DataLayerInstance) { return DataLayerInstance->CanUserRemoveActors(); });
 				FText RemoveActorToolTip = bSelectedActorsAreAllUserManaged ? FText::GetEmpty() : SelectionContainsNonUserManagedActorsText;
-				Section.AddMenuEntry("RemoveSelectedActorsFromSelectedDataLayers", LOCTEXT("RemoveSelectedActorsFromSelectedDataLayersMenu", "Remove Selected Actor(s) from Selected Data Layer(s)"), RemoveActorToolTip, FSlateIcon(),
+				Section.AddMenuEntry("RemoveSelectedActorsFromSelectedDataLayers", LOCTEXT("RemoveSelectedActorsFromSelectedDataLayersMenu", "Remove Selected Actor(s) from Selected Data Layer(s)"), RemoveActorToolTip, FSlateIcon(FCoreStyle::Get().GetStyleSetName(), "DataLayerBrowser.RemoveSelection"),
 					FUIAction(
 						FExecuteAction::CreateLambda([SelectedDataLayers]() 
 						{
@@ -1682,7 +1711,7 @@ void FDataLayerMode::RegisterContextMenu()
 			{
 				FToolMenuSection& Section = InMenu->AddSection("DataLayerSelection", LOCTEXT("DataLayerSelection", "Selection"));
 
-				Section.AddMenuEntry("SelectActorsInDataLayers", LOCTEXT("SelectActorsInDataLayers", "Select Actor(s) in Data Layer(s)"), FText(), FSlateIcon(),
+				Section.AddMenuEntry("SelectActorsInDataLayers", LOCTEXT("SelectActorsInDataLayers", "Select Actor(s) in Data Layer(s)"), FText(), FSlateIcon(FCoreStyle::Get().GetStyleSetName(), "GenericCommands.SelectAll"),
 					FUIAction(
 						FExecuteAction::CreateLambda([SelectedDataLayers]() {
 						check(!SelectedDataLayers.IsEmpty());
@@ -1694,7 +1723,7 @@ void FDataLayerMode::RegisterContextMenu()
 						FCanExecuteAction::CreateLambda([SelectedDataLayers] { return !SelectedDataLayers.IsEmpty(); })
 					));
 
-				Section.AddMenuEntry("AddActorsToSelection", LOCTEXT("AddActorsToSelection", "Add Actors in Data Layer(s) to Selection"), FText(), FSlateIcon(),
+				Section.AddMenuEntry("AddActorsToSelection", LOCTEXT("AddActorsToSelection", "Add Actors in Data Layer(s) to Selection"), FText(), FSlateIcon(FCoreStyle::Get().GetStyleSetName(), "Icons.Merge"),
 					FUIAction(
 						FExecuteAction::CreateLambda([SelectedDataLayers]() {
 						check(!SelectedDataLayers.IsEmpty());
@@ -1705,7 +1734,7 @@ void FDataLayerMode::RegisterContextMenu()
 						FCanExecuteAction::CreateLambda([SelectedDataLayers] { return !SelectedDataLayers.IsEmpty(); })
 					));
 
-				Section.AddMenuEntry("DeselectActors", LOCTEXT("DeselectActors", "Deselect Actor(s) in Data Layer(s)"), FText(), FSlateIcon(),
+				Section.AddMenuEntry("DeselectActors", LOCTEXT("DeselectActors", "Deselect Actor(s) in Data Layer(s)"), FText(), FSlateIcon(FCoreStyle::Get().GetStyleSetName(), "Icons.XCircle"),
 					FUIAction(
 						FExecuteAction::CreateLambda([SelectedDataLayers]() {
 						check(!SelectedDataLayers.IsEmpty());
@@ -1720,7 +1749,7 @@ void FDataLayerMode::RegisterContextMenu()
 			{
 				FToolMenuSection& Section = InMenu->AddSection("DataLayerLoading", LOCTEXT("DataLayerLoading", "Loading"));
 
-				Section.AddMenuEntry("ForceLoadActorsInDataLayers", LOCTEXT("ForceLoadActorsInDataLayers", "Force Load Actor(s) in Data Layer(s)"), FText(), FSlateIcon(),
+				Section.AddMenuEntry("ForceLoadActorsInDataLayers", LOCTEXT("ForceLoadActorsInDataLayers", "Force Load Actor(s) in Data Layer(s)"), FText(), FSlateIcon(FCoreStyle::Get().GetStyleSetName(), "Icons.Pinned"),
 					FUIAction(
 						FExecuteAction::CreateLambda([SelectedDataLayers]() 
 						{
@@ -1731,7 +1760,7 @@ void FDataLayerMode::RegisterContextMenu()
 						FCanExecuteAction::CreateLambda([SelectedDataLayers] { return !SelectedDataLayers.IsEmpty(); })
 					));
 
-				Section.AddMenuEntry("ReleaseForceLoadActorsInDataLayers", LOCTEXT("ReleaseForceLoadActorsInDataLayers", "Release Force Load Actor(s) in Data Layer(s)"), FText(), FSlateIcon(),
+				Section.AddMenuEntry("ReleaseForceLoadActorsInDataLayers", LOCTEXT("ReleaseForceLoadActorsInDataLayers", "Release Force Load Actor(s) in Data Layer(s)"), FText(), FSlateIcon(FCoreStyle::Get().GetStyleSetName(), "Icons.Unpinned"),
 					FUIAction(
 						FExecuteAction::CreateLambda([SelectedDataLayers]()
 							{
@@ -1746,7 +1775,7 @@ void FDataLayerMode::RegisterContextMenu()
 			{
 				FToolMenuSection& Section = InMenu->AddSection("DataLayerExpansion", LOCTEXT("DataLayerExpansion", "Expansion"));
 
-				Section.AddMenuEntry("CollapseAllDataLayers", LOCTEXT("CollapseAllDataLayers", "Collapse All Data Layers"), FText(), FSlateIcon(),
+				Section.AddMenuEntry("CollapseAllDataLayers", LOCTEXT("CollapseAllDataLayers", "Collapse All Data Layers"), FText(), FSlateIcon(FCoreStyle::Get().GetStyleSetName(), "Icons.ChevronUp"),
 					FUIAction(
 						FExecuteAction::CreateLambda([SceneOutliner,SelectedDataLayers]() {
 						check(!SelectedDataLayers.IsEmpty());
@@ -1757,7 +1786,7 @@ void FDataLayerMode::RegisterContextMenu()
 						FCanExecuteAction::CreateLambda([SelectedDataLayers] { return !SelectedDataLayers.IsEmpty(); })
 					));
 
-				Section.AddMenuEntry("ExpandAllDataLayers", LOCTEXT("ExpandAllDataLayers", "Expand All Data Layers"), FText(), FSlateIcon(),
+				Section.AddMenuEntry("ExpandAllDataLayers", LOCTEXT("ExpandAllDataLayers", "Expand All Data Layers"), FText(), FSlateIcon(FCoreStyle::Get().GetStyleSetName(), "Icons.ChevronDown"),
 					FUIAction(
 						FExecuteAction::CreateLambda([SceneOutliner,SelectedDataLayers]() {
 						check(!SelectedDataLayers.IsEmpty());
@@ -1772,7 +1801,7 @@ void FDataLayerMode::RegisterContextMenu()
 			{
 				FToolMenuSection& Section = InMenu->AddSection("DataLayerActorEditorContext", LOCTEXT("DataLayerActorEditorContext", "Actor Editor Context"));
 
-				Section.AddMenuEntry("MakeCurrentDataLayers", LOCTEXT("MakeCurrentDataLayers", "Make Current Data Layer(s)"), FText(), FSlateIcon(),
+				Section.AddMenuEntry("MakeCurrentDataLayers", LOCTEXT("MakeCurrentDataLayers", "Make Current Data Layer(s)"), FText(), FSlateIcon(FCoreStyle::Get().GetStyleSetName(), "Icons.Check"),
 					FUIAction(
 						FExecuteAction::CreateLambda([Mode, SceneOutliner, SelectedDataLayers]() 
 						{
@@ -1794,7 +1823,7 @@ void FDataLayerMode::RegisterContextMenu()
 						})
 					));
 
-				Section.AddMenuEntry("RemoveCurrentDataLayers", LOCTEXT("RemoveCurrentDataLayers", "Remove Current Data Layer(s)"), FText(), FSlateIcon(),
+				Section.AddMenuEntry("RemoveCurrentDataLayers", LOCTEXT("RemoveCurrentDataLayers", "Remove Current Data Layer(s)"), FText(), FSlateIcon(FCoreStyle::Get().GetStyleSetName(), "Icons.X"),
 					FUIAction(
 						FExecuteAction::CreateLambda([Mode, SceneOutliner, SelectedDataLayers]()
 						{
@@ -1816,7 +1845,7 @@ void FDataLayerMode::RegisterContextMenu()
 						})
 					));
 
-				Section.AddMenuEntry("ClearCurrentDataLayers", LOCTEXT("ClearCurrentDataLayers", "Clear Current Data Layer(s)"), FText(), FSlateIcon(),
+				Section.AddMenuEntry("ClearCurrentDataLayers", LOCTEXT("ClearCurrentDataLayers", "Clear Current Data Layer(s)"), FText(), FSlateIcon(FCoreStyle::Get().GetStyleSetName(), "Icons.Delete"),
 					FUIAction(
 						FExecuteAction::CreateLambda([AllDataLayers]() 
 						{
@@ -1836,7 +1865,7 @@ void FDataLayerMode::RegisterContextMenu()
 			{
 				FToolMenuSection& Section = InMenu->AddSection("DataLayerVisibility", LOCTEXT("DataLayerVisibility", "Visibility"));
 
-				Section.AddMenuEntry("MakeAllDataLayersVisible", LOCTEXT("MakeAllDataLayersVisible", "Make All Data Layers Visible"), FText(), FSlateIcon(),
+				Section.AddMenuEntry("MakeAllDataLayersVisible", LOCTEXT("MakeAllDataLayersVisible", "Make All Data Layers Visible"), FText(), FSlateIcon(FCoreStyle::Get().GetStyleSetName(), "Icons.Visible"),
 					FUIAction(
 						FExecuteAction::CreateLambda([AllDataLayers]()
 						{
@@ -2400,4 +2429,4 @@ void FDataLayerPickingMode::OnItemSelectionChanged(FSceneOutlinerTreeItemPtr Tre
 	}
 }
 
-#undef LOCTEXT_NAMESPACE
+#undef LOCTEXT_NAMESPACE
\ No newline at end of file
