[World] Deprecate ActorDataLayer property customization class, no longer needed since the property was deprecated.

#tests Tested opening blueprint that previously crashed

#Synced-CL 50064132

diff --git a/Engine/Source/Editor/DataLayerEditor/Private/DataLayer/DataLayerPropertyTypeCustomization.cpp b/Engine/Source/Editor/DataLayerEditor/Private/DataLayer/DataLayerPropertyTypeCustomization.cpp
deleted file mode 100644
diff --git a/Engine/Source/Editor/DataLayerEditor/Private/DataLayer/DataLayerPropertyTypeCustomization.h b/Engine/Source/Editor/DataLayerEditor/Private/DataLayer/DataLayerPropertyTypeCustomization.h
--- a/Engine/Source/Editor/DataLayerEditor/Private/DataLayer/DataLayerPropertyTypeCustomization.h
+++ b/Engine/Source/Editor/DataLayerEditor/Private/DataLayer/DataLayerPropertyTypeCustomization.h
@@ -21,27 +21,10 @@ struct FGeometry;
 struct FSlateBrush;
 struct FSlateColor;
 
+struct UE_DEPRECATED(5.8, "The ActorDataLayer property on AActor is no longer supported") FDataLayerPropertyTypeCustomization;
+
 struct FDataLayerPropertyTypeCustomization : public IPropertyTypeCustomization
 {
-	virtual void CustomizeHeader(TSharedRef<IPropertyHandle> StructPropertyHandle, FDetailWidgetRow& HeaderRow, IPropertyTypeCustomizationUtils& StructCustomizationUtils) override;
+	virtual void CustomizeHeader(TSharedRef<IPropertyHandle> StructPropertyHandle, FDetailWidgetRow& HeaderRow, IPropertyTypeCustomizationUtils& StructCustomizationUtils) override {}
 	virtual void CustomizeChildren(TSharedRef<IPropertyHandle> StructPropertyHandle, IDetailChildrenBuilder& StructBuilder, IPropertyTypeCustomizationUtils& StructCustomizationUtils) override {}
-
-private:
-
-	void AssignDataLayer(const UDataLayerInstance* InDataLayer);
-
-	UDataLayerInstance* GetDataLayerFromPropertyHandle(FPropertyAccess::Result* OutPropertyAccessResult = nullptr) const;
-	FText GetDataLayerText() const;
-	FSlateColor GetForegroundColor() const;
-	const FSlateBrush* GetDataLayerIcon() const;
-	EVisibility GetSelectDataLayerVisibility() const;
-
-	TSharedRef<SWidget> OnGetDataLayerMenu();
-	void OnBrowse();
-	FReply OnSelectDataLayer();
-	FReply OnDrop(const FGeometry& InGeometry, const FDragDropEvent& InDragDropEvent);
-	bool OnVerifyDrag(TSharedPtr<FDragDropOperation> InDragDrop);
-	TSharedPtr<const FDataLayerDragDropOp> GetDataLayerDragDropOp(TSharedPtr<FDragDropOperation> InDragDrop);
-
-	TSharedPtr<IPropertyHandle> PropertyHandle;
 };
diff --git a/Engine/Source/Editor/DataLayerEditor/Private/DataLayerEditorModule.cpp b/Engine/Source/Editor/DataLayerEditor/Private/DataLayerEditorModule.cpp
--- a/Engine/Source/Editor/DataLayerEditor/Private/DataLayerEditorModule.cpp
+++ b/Engine/Source/Editor/DataLayerEditor/Private/DataLayerEditorModule.cpp
@@ -38,13 +38,11 @@ IMPLEMENT_MODULE(FDataLayerEditorModule, DataLayerEditor );
 
 #define LOCTEXT_NAMESPACE "DataLayerEditorModule"
 
-static const FName NAME_ActorDataLayer(TEXT("ActorDataLayer"));
 static const FName NAME_DataLayerInstance(TEXT("DataLayerInstance"));
 
 void FDataLayerEditorModule::StartupModule()
 {
 	FPropertyEditorModule& PropertyModule = FModuleManager::LoadModuleChecked<FPropertyEditorModule>("PropertyEditor");
-	PropertyModule.RegisterCustomPropertyTypeLayout(NAME_ActorDataLayer, FOnGetPropertyTypeCustomizationInstance::CreateLambda([] { return MakeShared<FDataLayerPropertyTypeCustomization>(); }));
 	PropertyModule.RegisterCustomClassLayout(NAME_DataLayerInstance, FOnGetDetailCustomizationInstance::CreateStatic(&FDataLayerInstanceDetails::MakeInstance));
 
 	FEditorWidgetsModule& EditorWidgetsModule = FModuleManager::LoadModuleChecked<FEditorWidgetsModule>("EditorWidgets");
@@ -60,7 +58,6 @@ void FDataLayerEditorModule::ShutdownModule()
 
 	if (FPropertyEditorModule* PropertyModule = FModuleManager::GetModulePtr<FPropertyEditorModule>("PropertyEditor"))
 	{
-		PropertyModule->UnregisterCustomPropertyTypeLayout(NAME_ActorDataLayer);
 		PropertyModule->UnregisterCustomPropertyTypeLayout(NAME_DataLayerInstance);
 	}
 }
